UnrealVS v1.30: various minor stability fixes

#ttp 331470 - UnrealVS: toolbar combos disabled

#branch UE4
#proj UnrealVS
#summary

#change Minor refactor of BatchBuilder class to avoid issues when starting UnrealVS and loading options before the batch builder window opens.
#change Improved tick thread startup and shutdown. It only runs when a solution is open. This allows it to shutdown cleaner when VS closes.
#change Allowed the batch builder busy state to block solution closing. This should be prevented by VS anyway as a build would be running.
#change Minor refactor of the way CommandLineEditor and the StartupProjectSelector enable/disable and update their combo controls.
#change Version to 1.30

diff --git a/Engine/Source/Programs/UnrealVS/BatchBuilder.cs b/Engine/Source/Programs/UnrealVS/BatchBuilder.cs
--- a/Engine/Source/Programs/UnrealVS/BatchBuilder.cs
+++ b/Engine/Source/Programs/UnrealVS/BatchBuilder.cs
@@ -18,6 +18,26 @@ namespace UnrealVS
 
 		private const int BatchBuilderToolWindowId = 0x1300;
 
+		/** properties */
+
+		public bool IsBusy
+		{
+			get
+			{
+				return ToolControl.IsBusy;
+			}
+		}
+
+		private static BatchBuilderToolControl _ToolControl;
+		public static BatchBuilderToolControl ToolControl
+		{
+			get 
+			{
+				if (_ToolControl == null) _ToolControl = new BatchBuilderToolControl();
+				return _ToolControl;
+			}
+		}
+
 		/** methods */
 
 		public BatchBuilder()
@@ -34,12 +54,7 @@ namespace UnrealVS
 		/// <param name="Stream">The stream to load the option data from.</param>
 		public void LoadOptions(Stream Stream)
 		{
-			BatchBuilderToolControl ToolControl = GetToolControl();
-
-			if (ToolControl != null)
-			{
-				ToolControl.LoadOptions(Stream);
-			}
+			ToolControl.LoadOptions(Stream);
 		}
 
 		/// <summary>
@@ -48,12 +63,7 @@ namespace UnrealVS
 		/// <param name="Stream">The stream to save the option data to.</param>
 		public void SaveOptions(Stream Stream)
 		{
-			BatchBuilderToolControl ToolControl = GetToolControl();
-
-			if (ToolControl != null)
-			{
-				ToolControl.SaveOptions(Stream);
-			}
+			ToolControl.SaveOptions(Stream);
 		}
 
 		/// <summary>
@@ -61,12 +71,7 @@ namespace UnrealVS
 		/// </summary>
 		public void Tick()
 		{
-			BatchBuilderToolControl ToolControl = GetToolControl();
-
-			if (ToolControl != null)
-			{
-				ToolControl.Tick();
-			}
+			ToolControl.Tick();
 		}
 
 		private void ShowToolWindow(object sender, EventArgs e)
@@ -82,21 +87,5 @@ namespace UnrealVS
 			IVsWindowFrame ToolWindowFrame = (IVsWindowFrame)ToolWindow.Frame;
 			Microsoft.VisualStudio.ErrorHandler.ThrowOnFailure(ToolWindowFrame.Show());
 		}
-
-		private BatchBuilderToolControl GetToolControl()
-		{
-			BatchBuilderToolWindow ToolWindow =
-				UnrealVSPackage.Instance.FindToolWindow(typeof(BatchBuilderToolWindow), 0, true) as BatchBuilderToolWindow;
-			if ((null == ToolWindow) || (null == ToolWindow.Frame))
-			{
-				throw new NotSupportedException(Resources.ToolWindowCreateError);
-			}
-
-			if (ToolWindow != null)
-			{
-				return ToolWindow.Content as BatchBuilderToolControl;
-			}
-			return null;
-		}
 	}
 }
diff --git a/Engine/Source/Programs/UnrealVS/BatchBuilderToolWindow.cs b/Engine/Source/Programs/UnrealVS/BatchBuilderToolWindow.cs
--- a/Engine/Source/Programs/UnrealVS/BatchBuilderToolWindow.cs
+++ b/Engine/Source/Programs/UnrealVS/BatchBuilderToolWindow.cs
@@ -32,7 +32,7 @@ namespace UnrealVS
             // This is the user control hosted by the tool window; Note that, even if this class implements IDisposable,
             // we are not calling Dispose on this object. This is because ToolWindowPane calls Dispose on 
             // the object returned by the Content property.
-            base.Content = new BatchBuilderToolControl();
+			base.Content = BatchBuilder.ToolControl;
         }
 	}
 }
diff --git a/Engine/Source/Programs/UnrealVS/CommandLineEditor.cs b/Engine/Source/Programs/UnrealVS/CommandLineEditor.cs
--- a/Engine/Source/Programs/UnrealVS/CommandLineEditor.cs
+++ b/Engine/Source/Programs/UnrealVS/CommandLineEditor.cs
@@ -35,32 +35,24 @@ namespace UnrealVS
 				// CommandLineCombo
 				var ComboCommandID = new CommandID(GuidList.UnrealVSCmdSet, ComboID);
 				ComboCommand = new OleMenuCommand(new EventHandler(ComboHandler), ComboCommandID);
-				ComboCommand.BeforeQueryStatus += (sender, args) => { ComboCommand.Enabled = UnrealVSPackage.Instance.DTE.Solution.IsOpen; };
 				UnrealVSPackage.Instance.MenuCommandService.AddCommand(ComboCommand);
 
 				// CommandLineComboList
 				var ComboListCommandID = new CommandID(GuidList.UnrealVSCmdSet, ComboListID);
-				var ComboListCommand = new OleMenuCommand(new EventHandler(ComboListHandler), ComboListCommandID);
-				UnrealVSPackage.Instance.MenuCommandService.AddCommand(ComboListCommand);
+				ComboCommandList = new OleMenuCommand(new EventHandler(ComboListHandler), ComboListCommandID);
+				UnrealVSPackage.Instance.MenuCommandService.AddCommand(ComboCommandList);
 			}
 
-
 			// Register for events that we care about
-			UnrealVSPackage.Instance.OnStartupProjectChanged += OnStartupProjectChanged;
+			UnrealVSPackage.Instance.OnSolutionOpened += UpdateCommandLineCombo;
+			UnrealVSPackage.Instance.OnSolutionClosed += UpdateCommandLineCombo;
+			UnrealVSPackage.Instance.OnStartupProjectChanged += (p) => UpdateCommandLineCombo();
 			UnrealVSPackage.Instance.OnStartupProjectPropertyChanged += OnStartupProjectPropertyChanged;
 			UnrealVSPackage.Instance.OnStartupProjectConfigChanged += OnStartupProjectConfigChanged;
 
 			UpdateCommandLineCombo();
 		}
 
-		/// <summary>
-		/// Called when the startup project has changed to a new project.  We'll refresh our interface.
-		/// </summary>
-		public void OnStartupProjectChanged( Project NewStartupProject )
-		{
-			UpdateCommandLineCombo();
-		}
-
 		/// <summary>
 		/// Called when a startup project property has changed.  We'll refresh our interface.
 		/// </summary>
@@ -130,10 +122,12 @@ namespace UnrealVS
 			if( ProjectHierarchy != null )
 			{
 				ComboCommand.Enabled = true;
+				ComboCommandList.Enabled = true;
 			}
 			else
 			{
 				ComboCommand.Enabled = false;
+				ComboCommandList.Enabled = false;
 			}
 
 			string CommandLine = MakeCommandLineComboText();
@@ -338,6 +332,7 @@ namespace UnrealVS
 
 		/// Combo control for command-line editor
 		private OleMenuCommand ComboCommand;
+		private OleMenuCommand ComboCommandList;
 
 		/// Used to store the user edited commandline mid-edit, in the combo handler
 		private string DesiredCommandLine;
diff --git a/Engine/Source/Programs/UnrealVS/Properties/AssemblyInfo.cs b/Engine/Source/Programs/UnrealVS/Properties/AssemblyInfo.cs
--- a/Engine/Source/Programs/UnrealVS/Properties/AssemblyInfo.cs
+++ b/Engine/Source/Programs/UnrealVS/Properties/AssemblyInfo.cs
@@ -31,8 +31,8 @@ using System.Runtime.InteropServices;
 // You can specify all the values or you can default the Revision and Build Numbers 
 // by using the '*' as shown below:
 
-[assembly: AssemblyVersion("1.29.0.0")]
-[assembly: AssemblyFileVersion("1.29.0.0")]
+[assembly: AssemblyVersion("1.30.0.0")]
+[assembly: AssemblyFileVersion("1.30.0.0")]
 
 
 
diff --git a/Engine/Source/Programs/UnrealVS/StartupProjectSelector.cs b/Engine/Source/Programs/UnrealVS/StartupProjectSelector.cs
--- a/Engine/Source/Programs/UnrealVS/StartupProjectSelector.cs
+++ b/Engine/Source/Programs/UnrealVS/StartupProjectSelector.cs
@@ -43,18 +43,23 @@ namespace UnrealVS
 
 		public StartupProjectSelector()
 		{
+			_bIsSolutionOpened = UnrealVSPackage.Instance.DTE.Solution.IsOpen;
+
 			// Create the handlers for our commands
 			{
 				// StartupProjectCombo
 				var StartupProjectComboCommandId = new CommandID( GuidList.UnrealVSCmdSet, StartupProjectComboId );
-				StartupProjectComboCommand = new OleMenuCommand( StartupProjectComboHandler, StartupProjectComboCommandId );
-				StartupProjectComboCommand.BeforeQueryStatus += (sender, args) => { StartupProjectComboCommand.Enabled = _bIsSolutionOpened; };
+				var StartupProjectComboCommand = new OleMenuCommand(StartupProjectComboHandler, StartupProjectComboCommandId);
+				StartupProjectComboCommand.BeforeQueryStatus += (sender, args) => StartupProjectComboCommand.Enabled = _bIsSolutionOpened;
 				UnrealVSPackage.Instance.MenuCommandService.AddCommand( StartupProjectComboCommand );
 
 				// StartupProjectComboList
 				var StartupProjectComboListCommandId = new CommandID( GuidList.UnrealVSCmdSet, StartupProjectComboListId );
-				var StartupProjectComboListCommand = new OleMenuCommand( StartupProjectComboListHandler, StartupProjectComboListCommandId );
+				var StartupProjectComboListCommand = new OleMenuCommand(StartupProjectComboListHandler, StartupProjectComboListCommandId);
+				StartupProjectComboListCommand.BeforeQueryStatus += (sender, args) => StartupProjectComboListCommand.Enabled = _bIsSolutionOpened;
 				UnrealVSPackage.Instance.MenuCommandService.AddCommand( StartupProjectComboListCommand );
+
+				StartupProjectComboCommands = new[] { StartupProjectComboCommand, StartupProjectComboListCommand };
 			} 
 
 			// Register for events that we care about
@@ -64,6 +69,7 @@ namespace UnrealVS
 				{
 					Logging.WriteLine("Opened solution " + UnrealVSPackage.Instance.DTE.Solution.FullName);
 					_bIsSolutionOpened = true;
+					UpdateStartupProjectCombo();
 					UpdateStartupProjectList(Utils.GetAllProjectsFromDTE());
 				};
 			UnrealVSPackage.Instance.OnSolutionClosing +=
@@ -71,6 +77,7 @@ namespace UnrealVS
 				{
 					Logging.WriteLine("Closing solution");
 					_bIsSolutionOpened = false;
+					UpdateStartupProjectCombo();
 					_CachedStartupProjects.Clear();
 				};
 			UnrealVSPackage.Instance.OnProjectOpened +=
@@ -190,16 +197,9 @@ namespace UnrealVS
 		/// </summary>
 		private void UpdateStartupProjectCombo()
 		{
-			// Switch to this project!
-			IVsHierarchy ProjectHierarchy;
-			UnrealVSPackage.Instance.SolutionBuildManager.get_StartupProject( out ProjectHierarchy );
-			if( ProjectHierarchy != null )
-			{
-				StartupProjectComboCommand.Enabled = true;
-			}
-			else
+			foreach(var Command in StartupProjectComboCommands)
 			{
-				StartupProjectComboCommand.Enabled = false;
+				Command.Enabled = _bIsSolutionOpened;
 			}
 		}
 
@@ -275,8 +275,8 @@ namespace UnrealVS
 			Marshal.GetNativeVariantForObject(CachedStartupProjectNames, OleArgs.OutValue);
 		}
 
-		/// Combo control for our startup project selector
-		OleMenuCommand StartupProjectComboCommand;
+		/// Combo control and list commands for our startup project selector
+		readonly OleMenuCommand[] StartupProjectComboCommands;
 
 		/// Keep track of whether a solution is opened
 		bool _bIsSolutionOpened;
diff --git a/Engine/Source/Programs/UnrealVS/TargetManifests/VS11.source.extension.vsixmanifest b/Engine/Source/Programs/UnrealVS/TargetManifests/VS11.source.extension.vsixmanifest
--- a/Engine/Source/Programs/UnrealVS/TargetManifests/VS11.source.extension.vsixmanifest
+++ b/Engine/Source/Programs/UnrealVS/TargetManifests/VS11.source.extension.vsixmanifest
@@ -1,7 +1,7 @@
 ﻿<?xml version="1.0" encoding="utf-8"?>
 <PackageManifest Version="2.0.0" xmlns="http://schemas.microsoft.com/developer/vsx-schema/2011" xmlns:d="http://schemas.microsoft.com/developer/vsx-schema-design/2011">
 	<Metadata>
-		<Identity Id="ddbf523f-7eb6-4887-bd51-85a714ff87eb" Version="1.29" Language="en-US" Publisher="Epic Games, Inc." />
+		<Identity Id="ddbf523f-7eb6-4887-bd51-85a714ff87eb" Version="1.30" Language="en-US" Publisher="Epic Games, Inc." />
 		<DisplayName>UnrealVS</DisplayName>
 		<Description xml:space="preserve">Unreal Engine extension for Visual Studio</Description>
 		<MoreInfo>http://epicgames.com</MoreInfo>
diff --git a/Engine/Source/Programs/UnrealVS/TargetManifests/VS11d.source.extension.vsixmanifest b/Engine/Source/Programs/UnrealVS/TargetManifests/VS11d.source.extension.vsixmanifest
--- a/Engine/Source/Programs/UnrealVS/TargetManifests/VS11d.source.extension.vsixmanifest
+++ b/Engine/Source/Programs/UnrealVS/TargetManifests/VS11d.source.extension.vsixmanifest
@@ -1,7 +1,7 @@
 ﻿<?xml version="1.0" encoding="utf-8"?>
 <PackageManifest Version="2.0.0" xmlns="http://schemas.microsoft.com/developer/vsx-schema/2011" xmlns:d="http://schemas.microsoft.com/developer/vsx-schema-design/2011">
 	<Metadata>
-		<Identity Id="ddbf523f-7eb6-4887-bd51-85a714ff87eb" Version="1.29" Language="en-US" Publisher="Epic Games, Inc." />
+		<Identity Id="ddbf523f-7eb6-4887-bd51-85a714ff87eb" Version="1.30" Language="en-US" Publisher="Epic Games, Inc." />
 		<DisplayName>UnrealVS</DisplayName>
 		<Description xml:space="preserve">Unreal Engine extension for Visual Studio</Description>
 		<MoreInfo>http://epicgames.com</MoreInfo>
diff --git a/Engine/Source/Programs/UnrealVS/TargetManifests/VS12.source.extension.vsixmanifest b/Engine/Source/Programs/UnrealVS/TargetManifests/VS12.source.extension.vsixmanifest
--- a/Engine/Source/Programs/UnrealVS/TargetManifests/VS12.source.extension.vsixmanifest
+++ b/Engine/Source/Programs/UnrealVS/TargetManifests/VS12.source.extension.vsixmanifest
@@ -1,7 +1,7 @@
 ﻿<?xml version="1.0" encoding="utf-8"?>
 <PackageManifest Version="2.0.0" xmlns="http://schemas.microsoft.com/developer/vsx-schema/2011" xmlns:d="http://schemas.microsoft.com/developer/vsx-schema-design/2011">
 	<Metadata>
-		<Identity Id="ddbf523f-7eb6-4887-bd51-85a714ff87eb" Version="1.29" Language="en-US" Publisher="Epic Games, Inc." />
+		<Identity Id="ddbf523f-7eb6-4887-bd51-85a714ff87eb" Version="1.30" Language="en-US" Publisher="Epic Games, Inc." />
 		<DisplayName>UnrealVS</DisplayName>
 		<Description xml:space="preserve">Unreal Engine extension for Visual Studio</Description>
 		<MoreInfo>http://epicgames.com</MoreInfo>
diff --git a/Engine/Source/Programs/UnrealVS/TargetManifests/VS12d.source.extension.vsixmanifest b/Engine/Source/Programs/UnrealVS/TargetManifests/VS12d.source.extension.vsixmanifest
--- a/Engine/Source/Programs/UnrealVS/TargetManifests/VS12d.source.extension.vsixmanifest
+++ b/Engine/Source/Programs/UnrealVS/TargetManifests/VS12d.source.extension.vsixmanifest
@@ -1,7 +1,7 @@
 ﻿<?xml version="1.0" encoding="utf-8"?>
 <PackageManifest Version="2.0.0" xmlns="http://schemas.microsoft.com/developer/vsx-schema/2011" xmlns:d="http://schemas.microsoft.com/developer/vsx-schema-design/2011">
 	<Metadata>
-		<Identity Id="ddbf523f-7eb6-4887-bd51-85a714ff87eb" Version="1.29" Language="en-US" Publisher="Epic Games, Inc." />
+		<Identity Id="ddbf523f-7eb6-4887-bd51-85a714ff87eb" Version="1.30" Language="en-US" Publisher="Epic Games, Inc." />
 		<DisplayName>UnrealVS</DisplayName>
 		<Description xml:space="preserve">Unreal Engine extension for Visual Studio</Description>
 		<MoreInfo>http://epicgames.com</MoreInfo>
diff --git a/Engine/Source/Programs/UnrealVS/UnrealVSPackage.cs b/Engine/Source/Programs/UnrealVS/UnrealVSPackage.cs
--- a/Engine/Source/Programs/UnrealVS/UnrealVSPackage.cs
+++ b/Engine/Source/Programs/UnrealVS/UnrealVSPackage.cs
@@ -54,7 +54,7 @@ namespace UnrealVS
 	{
 		/** Constants */
 
-		private const string VersionString = "v1.29";
+		private const string VersionString = "v1.30";
 		private const string UnrealSolutionFileNamePrefix = "UE4";
 		private const string ExtensionName = "UnrealVS";
 		private const string CommandLineOptionKey = ExtensionName + "CommandLineMRU";
@@ -223,22 +223,45 @@ namespace UnrealVS
 			// Create the project menu quick builder
 			QuickBuilder = new QuickBuild();
 
+			// Call parent implementation
+			base.Initialize();
+
+			if (DTE.Solution.IsOpen)
+			{
+				StartTicker();
+			}
+		}
+
+		private void StartTicker()
+		{
 			// Create a "ticker" on a background thread that ticks the package on the UI thread
+			Interlocked.Exchange(ref bCancelTicker, 0);
 			Ticker = new Thread(TickAsyncMain);
 			Ticker.Priority = ThreadPriority.Lowest;
-			Ticker.Start();
+			Ticker.Start();			
+		}
 
-			// Call parent implementation
-			base.Initialize();
+		private void StopTicker()
+		{
+			if (bCancelTicker == 0)
+			{
+				Interlocked.Exchange(ref bCancelTicker, 1);
+			}
 		}
 
+		/// <summary>
+		/// Tick loop on worker thread
+		/// </summary>
 		private void TickAsyncMain()
 		{
 			try
 			{
-				while (bCancelTicker == 0)
+				while (true)
 				{
-					ThreadHelper.Generic.Invoke(Tick);
+					if (bCancelTicker != 0) return; 
+					ThreadHelper.Generic.BeginInvoke(Tick);
+
+					if (bCancelTicker != 0) return; 
 					Thread.Sleep(TickPeriod);
 				}
 			}
@@ -247,25 +270,26 @@ namespace UnrealVS
 			}
 		}
 
+		/// <summary>
+		/// Tick function on main UI thread
+		/// </summary>
 		private void Tick()
 		{
-			if (bCancelTicker == 0)
-			{
-				BatchBuilder.Tick();
-			}
+			BatchBuilder.Tick();
 		}
 
 		/// IDispose pattern lets us clean up our stuff!
 		protected override void Dispose( bool disposing )
 		{
-			// Shutdown the "ticker"
-			Interlocked.Increment(ref bCancelTicker);
-			if (!Ticker.Join(TickPeriod + TickPeriod))
+			if (Ticker != null && Ticker.IsAlive)
 			{
-				Ticker.Abort();
-				Ticker.Join();
+				Thread.Sleep(TickPeriod + TickPeriod);
+				if (Ticker.IsAlive)
+				{
+					Logging.WriteLine("WARNING: Force aborting Ticker thread");
+					Ticker.Abort();
+				}
 			}
-			Ticker = null;
 
 			base.Dispose(disposing);
 
@@ -528,6 +552,8 @@ namespace UnrealVS
 		{
 			UpdateUnrealLoadedStatus();
 
+			StartTicker();
+
             if (OnSolutionOpened != null)
             {
                 OnSolutionOpened();
@@ -552,6 +578,8 @@ namespace UnrealVS
 
 		int IVsSolutionEvents.OnBeforeCloseSolution( object pUnkReserved )
 		{
+			StopTicker();
+
             if (OnSolutionClosing != null)
             {
                 OnSolutionClosing();
@@ -572,6 +600,11 @@ namespace UnrealVS
 
 		int IVsSolutionEvents.OnQueryCloseSolution( object pUnkReserved, ref int pfCancel )
 		{
+			if (BatchBuilder.IsBusy)
+			{
+				pfCancel = 1;
+			}
+
 			return VSConstants.S_OK;
 		}
 
