UnrealVS v1.31: various minor stability fixes

diff --git a/Engine/Source/Programs/UnrealVS/Properties/AssemblyInfo.cs b/Engine/Source/Programs/UnrealVS/Properties/AssemblyInfo.cs
--- a/Engine/Source/Programs/UnrealVS/Properties/AssemblyInfo.cs
+++ b/Engine/Source/Programs/UnrealVS/Properties/AssemblyInfo.cs
@@ -31,8 +31,8 @@ using System.Runtime.InteropServices;
 // You can specify all the values or you can default the Revision and Build Numbers 
 // by using the '*' as shown below:
 
-[assembly: AssemblyVersion("1.30.0.0")]
-[assembly: AssemblyFileVersion("1.30.0.0")]
+[assembly: AssemblyVersion("1.31.0.0")]
+[assembly: AssemblyFileVersion("1.31.0.0")]
 
 
 
diff --git a/Engine/Source/Programs/UnrealVS/QuickBuild.cs b/Engine/Source/Programs/UnrealVS/QuickBuild.cs
--- a/Engine/Source/Programs/UnrealVS/QuickBuild.cs
+++ b/Engine/Source/Programs/UnrealVS/QuickBuild.cs
@@ -191,7 +191,10 @@ namespace UnrealVS
 				new SubMenu {Name = "PS4", SubMenuId = 0x1900, DynamicStartCommandId = 0x1930},
 				new SubMenu {Name = "IOS", SubMenuId = 0x1A00, DynamicStartCommandId = 0x1A30},
 				new SubMenu {Name = "Android", SubMenuId = 0x1B00, DynamicStartCommandId = 0x1B30},
-				new SubMenu {Name = "WinRT", SubMenuId = 0x1C00, DynamicStartCommandId = 0x1C30}
+				new SubMenu {Name = "WinRT", SubMenuId = 0x1C00, DynamicStartCommandId = 0x1C30},
+				new SubMenu {Name = "WinRT_ARM", SubMenuId = 0x1D00, DynamicStartCommandId = 0x1D30},
+				new SubMenu {Name = "HTML5", SubMenuId = 0x1E00, DynamicStartCommandId = 0x1E30},
+				new SubMenu {Name = "Linux", SubMenuId = 0x1F00, DynamicStartCommandId = 0x1F30},
 			};
 
 		/// The main root command of the Quick Build menu hierarchy - used to hide it when not active
diff --git a/Engine/Source/Programs/UnrealVS/TargetManifests/VS11.source.extension.vsixmanifest b/Engine/Source/Programs/UnrealVS/TargetManifests/VS11.source.extension.vsixmanifest
--- a/Engine/Source/Programs/UnrealVS/TargetManifests/VS11.source.extension.vsixmanifest
+++ b/Engine/Source/Programs/UnrealVS/TargetManifests/VS11.source.extension.vsixmanifest
@@ -1,7 +1,7 @@
 ﻿<?xml version="1.0" encoding="utf-8"?>
 <PackageManifest Version="2.0.0" xmlns="http://schemas.microsoft.com/developer/vsx-schema/2011" xmlns:d="http://schemas.microsoft.com/developer/vsx-schema-design/2011">
 	<Metadata>
-		<Identity Id="ddbf523f-7eb6-4887-bd51-85a714ff87eb" Version="1.30" Language="en-US" Publisher="Epic Games, Inc." />
+		<Identity Id="ddbf523f-7eb6-4887-bd51-85a714ff87eb" Version="1.31" Language="en-US" Publisher="Epic Games, Inc." />
 		<DisplayName>UnrealVS</DisplayName>
 		<Description xml:space="preserve">Unreal Engine extension for Visual Studio</Description>
 		<MoreInfo>http://epicgames.com</MoreInfo>
diff --git a/Engine/Source/Programs/UnrealVS/TargetManifests/VS11d.source.extension.vsixmanifest b/Engine/Source/Programs/UnrealVS/TargetManifests/VS11d.source.extension.vsixmanifest
--- a/Engine/Source/Programs/UnrealVS/TargetManifests/VS11d.source.extension.vsixmanifest
+++ b/Engine/Source/Programs/UnrealVS/TargetManifests/VS11d.source.extension.vsixmanifest
@@ -1,7 +1,7 @@
 ﻿<?xml version="1.0" encoding="utf-8"?>
 <PackageManifest Version="2.0.0" xmlns="http://schemas.microsoft.com/developer/vsx-schema/2011" xmlns:d="http://schemas.microsoft.com/developer/vsx-schema-design/2011">
 	<Metadata>
-		<Identity Id="ddbf523f-7eb6-4887-bd51-85a714ff87eb" Version="1.30" Language="en-US" Publisher="Epic Games, Inc." />
+		<Identity Id="ddbf523f-7eb6-4887-bd51-85a714ff87eb" Version="1.31" Language="en-US" Publisher="Epic Games, Inc." />
 		<DisplayName>UnrealVS</DisplayName>
 		<Description xml:space="preserve">Unreal Engine extension for Visual Studio</Description>
 		<MoreInfo>http://epicgames.com</MoreInfo>
diff --git a/Engine/Source/Programs/UnrealVS/TargetManifests/VS12.source.extension.vsixmanifest b/Engine/Source/Programs/UnrealVS/TargetManifests/VS12.source.extension.vsixmanifest
--- a/Engine/Source/Programs/UnrealVS/TargetManifests/VS12.source.extension.vsixmanifest
+++ b/Engine/Source/Programs/UnrealVS/TargetManifests/VS12.source.extension.vsixmanifest
@@ -1,7 +1,7 @@
 ﻿<?xml version="1.0" encoding="utf-8"?>
 <PackageManifest Version="2.0.0" xmlns="http://schemas.microsoft.com/developer/vsx-schema/2011" xmlns:d="http://schemas.microsoft.com/developer/vsx-schema-design/2011">
 	<Metadata>
-		<Identity Id="ddbf523f-7eb6-4887-bd51-85a714ff87eb" Version="1.30" Language="en-US" Publisher="Epic Games, Inc." />
+		<Identity Id="ddbf523f-7eb6-4887-bd51-85a714ff87eb" Version="1.31" Language="en-US" Publisher="Epic Games, Inc." />
 		<DisplayName>UnrealVS</DisplayName>
 		<Description xml:space="preserve">Unreal Engine extension for Visual Studio</Description>
 		<MoreInfo>http://epicgames.com</MoreInfo>
diff --git a/Engine/Source/Programs/UnrealVS/TargetManifests/VS12d.source.extension.vsixmanifest b/Engine/Source/Programs/UnrealVS/TargetManifests/VS12d.source.extension.vsixmanifest
--- a/Engine/Source/Programs/UnrealVS/TargetManifests/VS12d.source.extension.vsixmanifest
+++ b/Engine/Source/Programs/UnrealVS/TargetManifests/VS12d.source.extension.vsixmanifest
@@ -1,7 +1,7 @@
 ﻿<?xml version="1.0" encoding="utf-8"?>
 <PackageManifest Version="2.0.0" xmlns="http://schemas.microsoft.com/developer/vsx-schema/2011" xmlns:d="http://schemas.microsoft.com/developer/vsx-schema-design/2011">
 	<Metadata>
-		<Identity Id="ddbf523f-7eb6-4887-bd51-85a714ff87eb" Version="1.30" Language="en-US" Publisher="Epic Games, Inc." />
+		<Identity Id="ddbf523f-7eb6-4887-bd51-85a714ff87eb" Version="1.31" Language="en-US" Publisher="Epic Games, Inc." />
 		<DisplayName>UnrealVS</DisplayName>
 		<Description xml:space="preserve">Unreal Engine extension for Visual Studio</Description>
 		<MoreInfo>http://epicgames.com</MoreInfo>
diff --git a/Engine/Source/Programs/UnrealVS/UnrealVS.vsct b/Engine/Source/Programs/UnrealVS/UnrealVS.vsct
--- a/Engine/Source/Programs/UnrealVS/UnrealVS.vsct
+++ b/Engine/Source/Programs/UnrealVS/UnrealVS.vsct
@@ -196,6 +196,42 @@
 				</Strings>
 			</Button>
 
+			<!-- WinRT_ARM - Quick build dynamic placeholder button -->
+			<Button guid="UnrealVSCmdSetGUID" id="ProjectQuickBuildWinRTARMDynamicStartCommandID" priority="0x0000">
+				<Parent guid="UnrealVSCmdSetGUID" id="ProjectQuickBuildWinRTARMSubMenuGroupID" />
+				<CommandFlag>DynamicItemStart</CommandFlag>
+				<CommandFlag>DynamicVisibility</CommandFlag>
+				<CommandFlag>TextChanges</CommandFlag>
+				<!-- This text does not appear. -->
+				<Strings>
+					<ButtonText>DUMMY WinRT_ARM</ButtonText>
+				</Strings>
+			</Button>
+
+			<!-- HTML5 - Quick build dynamic placeholder button -->
+			<Button guid="UnrealVSCmdSetGUID" id="ProjectQuickBuildHTML5DynamicStartCommandID" priority="0x0000">
+				<Parent guid="UnrealVSCmdSetGUID" id="ProjectQuickBuildHTML5SubMenuGroupID" />
+				<CommandFlag>DynamicItemStart</CommandFlag>
+				<CommandFlag>DynamicVisibility</CommandFlag>
+				<CommandFlag>TextChanges</CommandFlag>
+				<!-- This text does not appear. -->
+				<Strings>
+					<ButtonText>DUMMY HTML5</ButtonText>
+				</Strings>
+			</Button>
+
+			<!-- Linux - Quick build dynamic placeholder button -->
+			<Button guid="UnrealVSCmdSetGUID" id="ProjectQuickBuildLinuxDynamicStartCommandID" priority="0x0000">
+				<Parent guid="UnrealVSCmdSetGUID" id="ProjectQuickBuildLinuxSubMenuGroupID" />
+				<CommandFlag>DynamicItemStart</CommandFlag>
+				<CommandFlag>DynamicVisibility</CommandFlag>
+				<CommandFlag>TextChanges</CommandFlag>
+				<!-- This text does not appear. -->
+				<Strings>
+					<ButtonText>DUMMY Linux</ButtonText>
+				</Strings>
+			</Button>
+			
 		</Buttons>
 	
 
@@ -384,7 +420,27 @@
 					<ButtonText>WinRT</ButtonText>
 				</Strings>
 			</Menu>
-      
+			<Menu guid="UnrealVSCmdSetGUID" id="ProjectQuickBuildWinRTARMSubMenuID" priority="0x0080" type="Menu">
+				<Parent guid="UnrealVSCmdSetGUID" id="ProjectQuickBuildSubMenuGroupID" />
+				<CommandFlag>DynamicVisibility</CommandFlag>
+				<Strings>
+					<ButtonText>WinRT_ARM</ButtonText>
+				</Strings>
+			</Menu>
+			<Menu guid="UnrealVSCmdSetGUID" id="ProjectQuickBuildHTML5SubMenuID" priority="0x0090" type="Menu">
+				<Parent guid="UnrealVSCmdSetGUID" id="ProjectQuickBuildSubMenuGroupID" />
+				<CommandFlag>DynamicVisibility</CommandFlag>
+				<Strings>
+					<ButtonText>HTML5</ButtonText>
+				</Strings>
+			</Menu>
+			<Menu guid="UnrealVSCmdSetGUID" id="ProjectQuickBuildLinuxSubMenuID" priority="0x0100" type="Menu">
+				<Parent guid="UnrealVSCmdSetGUID" id="ProjectQuickBuildSubMenuGroupID" />
+				<CommandFlag>DynamicVisibility</CommandFlag>
+				<Strings>
+					<ButtonText>Linux</ButtonText>
+				</Strings>
+			</Menu>
 		</Menus>
 
 		<!-- In this section you can define new menu groups. A menu group is a container for 
@@ -445,7 +501,19 @@
 				<!-- Parent the menu group to the WinRT quick build menu -->
 				<Parent guid="UnrealVSCmdSetGUID" id="ProjectQuickBuildWinRTSubMenuID"/>
 			</Group>
-      
+			<Group guid="UnrealVSCmdSetGUID" id="ProjectQuickBuildWinRTARMSubMenuGroupID" priority="0x0000">
+				<!-- Parent the menu group to the WinRT_ARM quick build menu -->
+				<Parent guid="UnrealVSCmdSetGUID" id="ProjectQuickBuildWinRTARMSubMenuID"/>
+			</Group>
+			<Group guid="UnrealVSCmdSetGUID" id="ProjectQuickBuildHTML5SubMenuGroupID" priority="0x0000">
+				<!-- Parent the menu group to the HTML5 quick build menu -->
+				<Parent guid="UnrealVSCmdSetGUID" id="ProjectQuickBuildHTML5SubMenuID"/>
+			</Group>
+			<Group guid="UnrealVSCmdSetGUID" id="ProjectQuickBuildLinuxSubMenuGroupID" priority="0x0000">
+				<!-- Parent the menu group to the Linux quick build menu -->
+				<Parent guid="UnrealVSCmdSetGUID" id="ProjectQuickBuildLinuxSubMenuID"/>
+			</Group>
+			
 		</Groups>
 
 		<Bitmaps>
@@ -526,6 +594,18 @@
 			<IDSymbol name="ProjectQuickBuildWinRTSubMenuGroupID" value="0x1C10" />
 			<IDSymbol name="ProjectQuickBuildWinRTDynamicStartCommandID" value="0x1C30" />
 
+			<IDSymbol name="ProjectQuickBuildWinRTARMSubMenuID" value="0x1D00" />
+			<IDSymbol name="ProjectQuickBuildWinRTARMSubMenuGroupID" value="0x1D10" />
+			<IDSymbol name="ProjectQuickBuildWinRTARMDynamicStartCommandID" value="0x1D30" />
+
+			<IDSymbol name="ProjectQuickBuildHTML5SubMenuID" value="0x1E00" />
+			<IDSymbol name="ProjectQuickBuildHTML5SubMenuGroupID" value="0x1E10" />
+			<IDSymbol name="ProjectQuickBuildHTML5DynamicStartCommandID" value="0x1E30" />
+
+			<IDSymbol name="ProjectQuickBuildLinuxSubMenuID" value="0x1F00" />
+			<IDSymbol name="ProjectQuickBuildLinuxSubMenuGroupID" value="0x1F10" />
+			<IDSymbol name="ProjectQuickBuildLinuxDynamicStartCommandID" value="0x1F30" />
+
 		</GuidSymbol>
 		
 		<GuidSymbol name="RefreshProjectsIconGUID" value="{8F4278B0-4C7D-4B67-9C27-AE415543EBE5}">
diff --git a/Engine/Source/Programs/UnrealVS/UnrealVSPackage.cs b/Engine/Source/Programs/UnrealVS/UnrealVSPackage.cs
--- a/Engine/Source/Programs/UnrealVS/UnrealVSPackage.cs
+++ b/Engine/Source/Programs/UnrealVS/UnrealVSPackage.cs
@@ -54,7 +54,7 @@ namespace UnrealVS
 	{
 		/** Constants */
 
-		private const string VersionString = "v1.30";
+		private const string VersionString = "v1.31";
 		private const string UnrealSolutionFileNamePrefix = "UE4";
 		private const string ExtensionName = "UnrealVS";
 		private const string CommandLineOptionKey = ExtensionName + "CommandLineMRU";
