UnrealVS v1.32: Improved identification of 'Game' projects by checking for uproject files matching the VS projects. Games are now always show at the top of the project list. Also fixed a command line editor issue.

diff --git a/Engine/Source/Programs/UnrealVS/CommandLineEditor.cs b/Engine/Source/Programs/UnrealVS/CommandLineEditor.cs
--- a/Engine/Source/Programs/UnrealVS/CommandLineEditor.cs
+++ b/Engine/Source/Programs/UnrealVS/CommandLineEditor.cs
@@ -162,10 +162,14 @@ namespace UnrealVS
 						if (UnrealVSPackage.Instance.IsUE4Loaded && Utils.IsGameProject(SelectedStartupProject) && Utils.HasUProjectCommandLineArg(ActiveConfiguration.Name))
 						{
 							string UProjectFileName = Utils.GetUProjectFileName(SelectedStartupProject);
-							if (CommandLineArguments.Trim().StartsWith(UProjectFileName, StringComparison.InvariantCultureIgnoreCase))
+							if (CommandLineArguments.Trim().StartsWith(UProjectFileName, StringComparison.OrdinalIgnoreCase))
 							{
 								CommandLineArguments = CommandLineArguments.Trim().Substring(UProjectFileName.Length).Trim();
 							}
+							else if (CommandLineArguments.Trim().StartsWith(SelectedStartupProject.Name, StringComparison.OrdinalIgnoreCase))
+							{
+								CommandLineArguments = CommandLineArguments.Trim().Substring(SelectedStartupProject.Name.Length).Trim();
+							}
 						}
 
 						Text = CommandLineArguments;
@@ -260,7 +264,16 @@ namespace UnrealVS
 						if (FullCommandLine.Trim().StartsWith(UProjectFileName, StringComparison.InvariantCultureIgnoreCase))
 						{
 							VsShellUtilities.ShowMessageBox(ServiceProvider.GlobalProvider,
-															string.Format("INFORMATION: The filename {0} has been removed from the command line because it is included automatically for 'Game' projects.", UProjectFileName),
+															string.Format("INFORMATION: The project filename {0} has been removed from the command line because it is included automatically for 'Game' projects.", UProjectFileName),
+															"UnrealVS",
+															OLEMSGICON.OLEMSGICON_INFO,
+															OLEMSGBUTTON.OLEMSGBUTTON_OK,
+															OLEMSGDEFBUTTON.OLEMSGDEFBUTTON_FIRST);
+						}
+						else if (FullCommandLine.Trim().StartsWith(SelectedStartupProject.Name, StringComparison.OrdinalIgnoreCase))
+						{
+							VsShellUtilities.ShowMessageBox(ServiceProvider.GlobalProvider,
+															string.Format("INFORMATION: The project name {0} has been removed from the command line because it is included automatically for 'Game' projects.", SelectedStartupProject.Name),
 															"UnrealVS",
 															OLEMSGICON.OLEMSGICON_INFO,
 															OLEMSGBUTTON.OLEMSGBUTTON_OK,
diff --git a/Engine/Source/Programs/UnrealVS/Properties/AssemblyInfo.cs b/Engine/Source/Programs/UnrealVS/Properties/AssemblyInfo.cs
--- a/Engine/Source/Programs/UnrealVS/Properties/AssemblyInfo.cs
+++ b/Engine/Source/Programs/UnrealVS/Properties/AssemblyInfo.cs
@@ -31,8 +31,8 @@ using System.Runtime.InteropServices;
 // You can specify all the values or you can default the Revision and Build Numbers 
 // by using the '*' as shown below:
 
-[assembly: AssemblyVersion("1.31.0.0")]
-[assembly: AssemblyFileVersion("1.31.0.0")]
+[assembly: AssemblyVersion("1.32.0.0")]
+[assembly: AssemblyFileVersion("1.32.0.0")]
 
 
 
diff --git a/Engine/Source/Programs/UnrealVS/TargetManifests/VS11.source.extension.vsixmanifest b/Engine/Source/Programs/UnrealVS/TargetManifests/VS11.source.extension.vsixmanifest
--- a/Engine/Source/Programs/UnrealVS/TargetManifests/VS11.source.extension.vsixmanifest
+++ b/Engine/Source/Programs/UnrealVS/TargetManifests/VS11.source.extension.vsixmanifest
@@ -1,7 +1,7 @@
 ﻿<?xml version="1.0" encoding="utf-8"?>
 <PackageManifest Version="2.0.0" xmlns="http://schemas.microsoft.com/developer/vsx-schema/2011" xmlns:d="http://schemas.microsoft.com/developer/vsx-schema-design/2011">
 	<Metadata>
-		<Identity Id="ddbf523f-7eb6-4887-bd51-85a714ff87eb" Version="1.31" Language="en-US" Publisher="Epic Games, Inc." />
+		<Identity Id="ddbf523f-7eb6-4887-bd51-85a714ff87eb" Version="1.32" Language="en-US" Publisher="Epic Games, Inc." />
 		<DisplayName>UnrealVS</DisplayName>
 		<Description xml:space="preserve">Unreal Engine extension for Visual Studio</Description>
 		<MoreInfo>http://epicgames.com</MoreInfo>
diff --git a/Engine/Source/Programs/UnrealVS/TargetManifests/VS11d.source.extension.vsixmanifest b/Engine/Source/Programs/UnrealVS/TargetManifests/VS11d.source.extension.vsixmanifest
--- a/Engine/Source/Programs/UnrealVS/TargetManifests/VS11d.source.extension.vsixmanifest
+++ b/Engine/Source/Programs/UnrealVS/TargetManifests/VS11d.source.extension.vsixmanifest
@@ -1,7 +1,7 @@
 ﻿<?xml version="1.0" encoding="utf-8"?>
 <PackageManifest Version="2.0.0" xmlns="http://schemas.microsoft.com/developer/vsx-schema/2011" xmlns:d="http://schemas.microsoft.com/developer/vsx-schema-design/2011">
 	<Metadata>
-		<Identity Id="ddbf523f-7eb6-4887-bd51-85a714ff87eb" Version="1.31" Language="en-US" Publisher="Epic Games, Inc." />
+		<Identity Id="ddbf523f-7eb6-4887-bd51-85a714ff87eb" Version="1.32" Language="en-US" Publisher="Epic Games, Inc." />
 		<DisplayName>UnrealVS</DisplayName>
 		<Description xml:space="preserve">Unreal Engine extension for Visual Studio</Description>
 		<MoreInfo>http://epicgames.com</MoreInfo>
diff --git a/Engine/Source/Programs/UnrealVS/TargetManifests/VS12.source.extension.vsixmanifest b/Engine/Source/Programs/UnrealVS/TargetManifests/VS12.source.extension.vsixmanifest
--- a/Engine/Source/Programs/UnrealVS/TargetManifests/VS12.source.extension.vsixmanifest
+++ b/Engine/Source/Programs/UnrealVS/TargetManifests/VS12.source.extension.vsixmanifest
@@ -1,7 +1,7 @@
 ﻿<?xml version="1.0" encoding="utf-8"?>
 <PackageManifest Version="2.0.0" xmlns="http://schemas.microsoft.com/developer/vsx-schema/2011" xmlns:d="http://schemas.microsoft.com/developer/vsx-schema-design/2011">
 	<Metadata>
-		<Identity Id="ddbf523f-7eb6-4887-bd51-85a714ff87eb" Version="1.31" Language="en-US" Publisher="Epic Games, Inc." />
+		<Identity Id="ddbf523f-7eb6-4887-bd51-85a714ff87eb" Version="1.32" Language="en-US" Publisher="Epic Games, Inc." />
 		<DisplayName>UnrealVS</DisplayName>
 		<Description xml:space="preserve">Unreal Engine extension for Visual Studio</Description>
 		<MoreInfo>http://epicgames.com</MoreInfo>
diff --git a/Engine/Source/Programs/UnrealVS/TargetManifests/VS12d.source.extension.vsixmanifest b/Engine/Source/Programs/UnrealVS/TargetManifests/VS12d.source.extension.vsixmanifest
--- a/Engine/Source/Programs/UnrealVS/TargetManifests/VS12d.source.extension.vsixmanifest
+++ b/Engine/Source/Programs/UnrealVS/TargetManifests/VS12d.source.extension.vsixmanifest
@@ -1,7 +1,7 @@
 ﻿<?xml version="1.0" encoding="utf-8"?>
 <PackageManifest Version="2.0.0" xmlns="http://schemas.microsoft.com/developer/vsx-schema/2011" xmlns:d="http://schemas.microsoft.com/developer/vsx-schema-design/2011">
 	<Metadata>
-		<Identity Id="ddbf523f-7eb6-4887-bd51-85a714ff87eb" Version="1.31" Language="en-US" Publisher="Epic Games, Inc." />
+		<Identity Id="ddbf523f-7eb6-4887-bd51-85a714ff87eb" Version="1.32" Language="en-US" Publisher="Epic Games, Inc." />
 		<DisplayName>UnrealVS</DisplayName>
 		<Description xml:space="preserve">Unreal Engine extension for Visual Studio</Description>
 		<MoreInfo>http://epicgames.com</MoreInfo>
diff --git a/Engine/Source/Programs/UnrealVS/UnrealVSPackage.cs b/Engine/Source/Programs/UnrealVS/UnrealVSPackage.cs
--- a/Engine/Source/Programs/UnrealVS/UnrealVSPackage.cs
+++ b/Engine/Source/Programs/UnrealVS/UnrealVSPackage.cs
@@ -54,7 +54,7 @@ namespace UnrealVS
 	{
 		/** Constants */
 
-		private const string VersionString = "v1.31";
+		private const string VersionString = "v1.32";
 		private const string UnrealSolutionFileNamePrefix = "UE4";
 		private const string ExtensionName = "UnrealVS";
 		private const string CommandLineOptionKey = ExtensionName + "CommandLineMRU";
