UnrealVS v1.33: Improved detection of Unreal Engine solution files by looking for a comment in the sln file added by UBT.

diff --git a/Engine/Source/Programs/UnrealBuildTool/System/VCProjectFileGenerator.cs b/Engine/Source/Programs/UnrealBuildTool/System/VCProjectFileGenerator.cs
--- a/Engine/Source/Programs/UnrealBuildTool/System/VCProjectFileGenerator.cs
+++ b/Engine/Source/Programs/UnrealBuildTool/System/VCProjectFileGenerator.cs
@@ -263,6 +263,7 @@ namespace UnrealBuildTool
 			// Setup solution file content
 			var VCSolutionFileContent = new StringBuilder();
 
+			const string VersionTag = "# UnrealEngineGeneratedSolutionVersion=1.0";
 
 			// Solution file header
 			if( ProjectFileFormat == VCProjectFileFormat.VisualStudio2013 )
@@ -270,7 +271,8 @@ namespace UnrealBuildTool
 				VCSolutionFileContent.Append(
 					ProjectFileGenerator.NewLine +
 					"Microsoft Visual Studio Solution File, Format Version 12.00" + ProjectFileGenerator.NewLine +
-					"# Visual Studio 2013" + ProjectFileGenerator.NewLine );
+					"# Visual Studio 2013" + ProjectFileGenerator.NewLine +
+					VersionTag + ProjectFileGenerator.NewLine);
 				
 				/* This is not required by VS 2013 to load the projects
 				VCSolutionFileContent.Append(
@@ -282,7 +284,8 @@ namespace UnrealBuildTool
 				VCSolutionFileContent.Append(
 					ProjectFileGenerator.NewLine +
 					"Microsoft Visual Studio Solution File, Format Version 12.00" + ProjectFileGenerator.NewLine +
-					"# Visual Studio 2012" + ProjectFileGenerator.NewLine );
+					"# Visual Studio 2012" + ProjectFileGenerator.NewLine +
+					VersionTag + ProjectFileGenerator.NewLine);
 			}
 			else
 			{
diff --git a/Engine/Source/Programs/UnrealVS/Properties/AssemblyInfo.cs b/Engine/Source/Programs/UnrealVS/Properties/AssemblyInfo.cs
--- a/Engine/Source/Programs/UnrealVS/Properties/AssemblyInfo.cs
+++ b/Engine/Source/Programs/UnrealVS/Properties/AssemblyInfo.cs
@@ -31,8 +31,8 @@ using System.Runtime.InteropServices;
 // You can specify all the values or you can default the Revision and Build Numbers 
 // by using the '*' as shown below:
 
-[assembly: AssemblyVersion("1.32.0.0")]
-[assembly: AssemblyFileVersion("1.32.0.0")]
+[assembly: AssemblyVersion("1.33.0.0")]
+[assembly: AssemblyFileVersion("1.33.0.0")]
 
 
 
diff --git a/Engine/Source/Programs/UnrealVS/StartupProjectSelector.cs b/Engine/Source/Programs/UnrealVS/StartupProjectSelector.cs
--- a/Engine/Source/Programs/UnrealVS/StartupProjectSelector.cs
+++ b/Engine/Source/Programs/UnrealVS/StartupProjectSelector.cs
@@ -130,7 +130,7 @@ namespace UnrealVS
 		{
 			foreach (Project Project in DirtyProjects)
 			{
-				if (ProjectFilter(Project))
+				if (ProjectPredicate(Project))
 				{
 					if (!_CachedStartupProjects.Any(ProjRef => ProjRef.Project == Project))
 					{
@@ -143,8 +143,7 @@ namespace UnrealVS
 				}
 			}
 
-			// Sort projects alphabetically
-			_CachedStartupProjects.Sort((A, B) => string.CompareOrdinal(A.Name, B.Name));
+			SortCachedStartupProjects();
 
 			if (StartupProjectListChanged != null)
 			{
@@ -152,7 +151,14 @@ namespace UnrealVS
 			}
 		}
 
-		private bool ProjectFilter(Project Project)
+		private void SortCachedStartupProjects()
+		{
+			// Sort projects by game y/n then alphabetically
+			_CachedStartupProjects = (_CachedStartupProjects.OrderBy(ProjectRef => Utils.IsGameProject(ProjectRef.Project) ? 0 : 1)
+				.ThenBy(ProjectRef => ProjectRef.Name)).ToList();
+		}
+
+		private bool ProjectPredicate(Project Project)
 		{
 			// Initialize the cached options if they haven't been read from UnrealVSOptions yet
 			if (!CachedHideNonGameStartupProjects.HasValue)
@@ -163,7 +169,7 @@ namespace UnrealVS
 			// Optionally, ignore non-game projects
 			if (CachedHideNonGameStartupProjects.Value)
 			{
-				if (!Project.Name.EndsWith("Game", StringComparison.InvariantCultureIgnoreCase))
+				if (!Utils.IsGameProject(Project))
 				{
 					Logging.WriteLine("StartupProjectSelector: Not listing project " + Project.Name + " because it is not a game");
 					return false;
@@ -171,7 +177,7 @@ namespace UnrealVS
 			}
 
 			// Always filter out non-executable projects
-			if (!Utils.IsProjectExecutable(Project))
+			if (!Utils.IsProjectSuitable(Project))
 			{
 				Logging.WriteLine("StartupProjectSelector: Not listing project " + Project.Name + " because it is not executable");
 				return false;
@@ -282,6 +288,6 @@ namespace UnrealVS
 		bool _bIsSolutionOpened;
 
 		/// Store the list of projects that should be shown in the dropdown list
-		readonly List<ProjectReference> _CachedStartupProjects = new List<ProjectReference>();
+		List<ProjectReference> _CachedStartupProjects = new List<ProjectReference>();
 	}
 }
diff --git a/Engine/Source/Programs/UnrealVS/TargetManifests/VS11.source.extension.vsixmanifest b/Engine/Source/Programs/UnrealVS/TargetManifests/VS11.source.extension.vsixmanifest
--- a/Engine/Source/Programs/UnrealVS/TargetManifests/VS11.source.extension.vsixmanifest
+++ b/Engine/Source/Programs/UnrealVS/TargetManifests/VS11.source.extension.vsixmanifest
@@ -1,7 +1,7 @@
 ﻿<?xml version="1.0" encoding="utf-8"?>
 <PackageManifest Version="2.0.0" xmlns="http://schemas.microsoft.com/developer/vsx-schema/2011" xmlns:d="http://schemas.microsoft.com/developer/vsx-schema-design/2011">
 	<Metadata>
-		<Identity Id="ddbf523f-7eb6-4887-bd51-85a714ff87eb" Version="1.32" Language="en-US" Publisher="Epic Games, Inc." />
+		<Identity Id="ddbf523f-7eb6-4887-bd51-85a714ff87eb" Version="1.33" Language="en-US" Publisher="Epic Games, Inc." />
 		<DisplayName>UnrealVS</DisplayName>
 		<Description xml:space="preserve">Unreal Engine extension for Visual Studio</Description>
 		<MoreInfo>http://epicgames.com</MoreInfo>
diff --git a/Engine/Source/Programs/UnrealVS/TargetManifests/VS11d.source.extension.vsixmanifest b/Engine/Source/Programs/UnrealVS/TargetManifests/VS11d.source.extension.vsixmanifest
--- a/Engine/Source/Programs/UnrealVS/TargetManifests/VS11d.source.extension.vsixmanifest
+++ b/Engine/Source/Programs/UnrealVS/TargetManifests/VS11d.source.extension.vsixmanifest
@@ -1,7 +1,7 @@
 ﻿<?xml version="1.0" encoding="utf-8"?>
 <PackageManifest Version="2.0.0" xmlns="http://schemas.microsoft.com/developer/vsx-schema/2011" xmlns:d="http://schemas.microsoft.com/developer/vsx-schema-design/2011">
 	<Metadata>
-		<Identity Id="ddbf523f-7eb6-4887-bd51-85a714ff87eb" Version="1.32" Language="en-US" Publisher="Epic Games, Inc." />
+		<Identity Id="ddbf523f-7eb6-4887-bd51-85a714ff87eb" Version="1.33" Language="en-US" Publisher="Epic Games, Inc." />
 		<DisplayName>UnrealVS</DisplayName>
 		<Description xml:space="preserve">Unreal Engine extension for Visual Studio</Description>
 		<MoreInfo>http://epicgames.com</MoreInfo>
diff --git a/Engine/Source/Programs/UnrealVS/TargetManifests/VS12.source.extension.vsixmanifest b/Engine/Source/Programs/UnrealVS/TargetManifests/VS12.source.extension.vsixmanifest
--- a/Engine/Source/Programs/UnrealVS/TargetManifests/VS12.source.extension.vsixmanifest
+++ b/Engine/Source/Programs/UnrealVS/TargetManifests/VS12.source.extension.vsixmanifest
@@ -1,7 +1,7 @@
 ﻿<?xml version="1.0" encoding="utf-8"?>
 <PackageManifest Version="2.0.0" xmlns="http://schemas.microsoft.com/developer/vsx-schema/2011" xmlns:d="http://schemas.microsoft.com/developer/vsx-schema-design/2011">
 	<Metadata>
-		<Identity Id="ddbf523f-7eb6-4887-bd51-85a714ff87eb" Version="1.32" Language="en-US" Publisher="Epic Games, Inc." />
+		<Identity Id="ddbf523f-7eb6-4887-bd51-85a714ff87eb" Version="1.33" Language="en-US" Publisher="Epic Games, Inc." />
 		<DisplayName>UnrealVS</DisplayName>
 		<Description xml:space="preserve">Unreal Engine extension for Visual Studio</Description>
 		<MoreInfo>http://epicgames.com</MoreInfo>
diff --git a/Engine/Source/Programs/UnrealVS/TargetManifests/VS12d.source.extension.vsixmanifest b/Engine/Source/Programs/UnrealVS/TargetManifests/VS12d.source.extension.vsixmanifest
--- a/Engine/Source/Programs/UnrealVS/TargetManifests/VS12d.source.extension.vsixmanifest
+++ b/Engine/Source/Programs/UnrealVS/TargetManifests/VS12d.source.extension.vsixmanifest
@@ -1,7 +1,7 @@
 ﻿<?xml version="1.0" encoding="utf-8"?>
 <PackageManifest Version="2.0.0" xmlns="http://schemas.microsoft.com/developer/vsx-schema/2011" xmlns:d="http://schemas.microsoft.com/developer/vsx-schema-design/2011">
 	<Metadata>
-		<Identity Id="ddbf523f-7eb6-4887-bd51-85a714ff87eb" Version="1.32" Language="en-US" Publisher="Epic Games, Inc." />
+		<Identity Id="ddbf523f-7eb6-4887-bd51-85a714ff87eb" Version="1.33" Language="en-US" Publisher="Epic Games, Inc." />
 		<DisplayName>UnrealVS</DisplayName>
 		<Description xml:space="preserve">Unreal Engine extension for Visual Studio</Description>
 		<MoreInfo>http://epicgames.com</MoreInfo>
diff --git a/Engine/Source/Programs/UnrealVS/UnrealVSPackage.cs b/Engine/Source/Programs/UnrealVS/UnrealVSPackage.cs
--- a/Engine/Source/Programs/UnrealVS/UnrealVSPackage.cs
+++ b/Engine/Source/Programs/UnrealVS/UnrealVSPackage.cs
@@ -5,6 +5,7 @@ using System.Collections.Generic;
 using System.Diagnostics;
 using System.Globalization;
 using System.IO;
+using System.Linq;
 using System.Runtime.InteropServices;
 using System.ComponentModel.Design;
 using System.Threading;
@@ -54,7 +55,7 @@ namespace UnrealVS
 	{
 		/** Constants */
 
-		private const string VersionString = "v1.32";
+		private const string VersionString = "v1.33";
 		private const string UnrealSolutionFileNamePrefix = "UE4";
 		private const string ExtensionName = "UnrealVS";
 		private const string CommandLineOptionKey = ExtensionName + "CommandLineMRU";
@@ -764,13 +765,40 @@ namespace UnrealVS
 
 		private void UpdateUnrealLoadedStatus()
 		{
+			if (!DTE.Solution.IsOpen)
+			{
+				IsUE4Loaded = false;
+				return;
+			}
+
 			string SolutionDirectory, UserOptsFile;
 			SolutionManager.GetSolutionInfo(out SolutionDirectory, out _SolutionFilepath, out UserOptsFile);
-			IsUE4Loaded =
-				(
-					_SolutionFilepath != null &&
-					Path.GetFileName(_SolutionFilepath).StartsWith(UnrealSolutionFileNamePrefix, StringComparison.InvariantCultureIgnoreCase)
-				);
+
+			var SolutionLines = new string[0];
+			try
+			{
+				SolutionLines = File.ReadAllLines(_SolutionFilepath);
+			}
+			catch
+			{
+			}
+
+			const string UBTTag = "# UnrealEngineGeneratedSolutionVersion=";
+			var UBTLine = SolutionLines.FirstOrDefault(TextLine => TextLine.Trim().StartsWith(UBTTag));
+			if (UBTLine != null)
+			{
+				_UBTVersion = UBTLine.Trim().Substring(UBTTag.Length);
+				IsUE4Loaded = true;
+			}
+			else
+			{
+				_UBTVersion = string.Empty;
+				IsUE4Loaded =
+					(
+						_SolutionFilepath != null &&
+						Path.GetFileName(_SolutionFilepath).StartsWith(UnrealSolutionFileNamePrefix, StringComparison.OrdinalIgnoreCase)
+					);
+			}
 		}
 
 		/** Private Fields & Properties */
@@ -807,6 +835,8 @@ namespace UnrealVS
 		/// Ticker thread cancel flag
 		private int bCancelTicker = 0;
 
+		private string _UBTVersion = string.Empty;
+
 		/// Obtains the DTE2 interface for this instance of VS from the RunningObjectTable
 		private static DTE2 GetDTE2ForCurrentInstance(DTE DTE)
 		{
diff --git a/Engine/Source/Programs/UnrealVS/Utils.cs b/Engine/Source/Programs/UnrealVS/Utils.cs
--- a/Engine/Source/Programs/UnrealVS/Utils.cs
+++ b/Engine/Source/Programs/UnrealVS/Utils.cs
@@ -14,6 +14,8 @@ namespace UnrealVS
 {
 	public class Utils
 	{
+		public const string UProjectExtension = "uproject";
+
 		public class SafeProjectReference
 		{
 			public string FullName { get; set; }
@@ -231,7 +233,7 @@ namespace UnrealVS
 		/// Helper to check the properties of a project and determine whether it can be run in VS.
 		/// Projects that return true can be run in the debugger by pressing the usual Start Debugging (F5) command.
 		/// </summary>
-		public static bool IsProjectExecutable(Project Project)
+		public static bool IsProjectSuitable(Project Project)
 		{
 			try
 			{
@@ -254,37 +256,40 @@ namespace UnrealVS
 					Logging.WriteLine("IsProjectExecutable: Warning - ActiveProjectConfig is null!");
 				}
 
-				bool IsExecutable = false;
+				bool IsSuitable = false;
 
 				if (Project.Kind.Equals(GuidList.VCSharpProjectKindGuidString, StringComparison.OrdinalIgnoreCase))
 				{
 					// C# project
 
-					Property StartActionProp = GetProjectConfigProperty(Project, null, "StartAction");
-					if (StartActionProp != null)
-					{
-						prjStartAction StartAction = (prjStartAction)StartActionProp.Value;
-						if (StartAction == prjStartAction.prjStartActionProject)
-						{
-							// Project starts the project's output file when run
-							Property OutputTypeProp = GetProjectProperty(Project, "OutputType");
-							if (OutputTypeProp != null)
-							{
-								prjOutputType OutputType = (prjOutputType)OutputTypeProp.Value;
-								if (OutputType == prjOutputType.prjOutputTypeWinExe ||
-									OutputType == prjOutputType.prjOutputTypeExe)
-								{
-									IsExecutable = true;
-								}
-							}
-						}
-						else if (StartAction == prjStartAction.prjStartActionProgram ||
-								 StartAction == prjStartAction.prjStartActionURL)
-						{
-							// Project starts an external program or a URL when run - assume it has been set deliberately to something executable
-							IsExecutable = true;
-						}
-					}
+					// Chris.Wood
+					//Property StartActionProp = GetProjectConfigProperty(Project, null, "StartAction");
+					//if (StartActionProp != null)
+					//{
+					//	prjStartAction StartAction = (prjStartAction)StartActionProp.Value;
+					//	if (StartAction == prjStartAction.prjStartActionProject)
+					//	{
+					//		// Project starts the project's output file when run
+					//		Property OutputTypeProp = GetProjectProperty(Project, "OutputType");
+					//		if (OutputTypeProp != null)
+					//		{
+					//			prjOutputType OutputType = (prjOutputType)OutputTypeProp.Value;
+					//			if (OutputType == prjOutputType.prjOutputTypeWinExe ||
+					//				OutputType == prjOutputType.prjOutputTypeExe)
+					//			{
+					//				IsSuitable = true;
+					//			}
+					//		}
+					//	}
+					//	else if (StartAction == prjStartAction.prjStartActionProgram ||
+					//			 StartAction == prjStartAction.prjStartActionURL)
+					//	{
+					//		// Project starts an external program or a URL when run - assume it has been set deliberately to something executable
+					//		IsSuitable = true;
+					//	}
+					//}
+
+					IsSuitable = true;
 				}
 				else if (Project.Kind.Equals(GuidList.VCProjectKindGuidString, StringComparison.OrdinalIgnoreCase))
 				{
@@ -314,7 +319,7 @@ namespace UnrealVS
 								if (VCConfigMatch.DebugAttach)
 								{
 									// Project attaches to a running process
-									IsExecutable = true;
+									IsSuitable = true;
 								}
 								else
 								{
@@ -326,7 +331,7 @@ namespace UnrealVS
 										if (VCConfigMatch.DebugRemoteCommand.Length != 0)
 										{
 											// An remote program is specified to run
-											IsExecutable = true;
+											IsSuitable = true;
 										}
 									}
 									else
@@ -336,7 +341,7 @@ namespace UnrealVS
 										if (VCConfigMatch.DebugCommand.Length != 0 && VCConfigMatch.DebugCommand != "$(TargetPath)")
 										{
 											// An external program is specified to run
-											IsExecutable = true;
+											IsSuitable = true;
 										}
 										else
 										{
@@ -344,7 +349,7 @@ namespace UnrealVS
 
 											if (VCConfigMatch.ConfigType == ConfigType.Application)
 											{
-												IsExecutable = true;
+												IsSuitable = true;
 											}
 											else if (VCConfigMatch.ConfigType == ConfigType.Generic)
 											{
@@ -355,7 +360,7 @@ namespace UnrealVS
 													string Ext = Path.GetExtension(VCConfigMatch.NMakeToolOutput);
 													if (!IsLibraryFileExtension(Ext))
 													{
-														IsExecutable = true;
+														IsSuitable = true;
 													}
 												}
 											}
@@ -372,7 +377,7 @@ namespace UnrealVS
 					Logging.WriteLine("IsProjectExecutable: Unrecognised 'Kind' in project " + Project.Name + " guid=" + Project.Kind);
 				}
 
-				return IsExecutable;
+				return IsSuitable;
 			}
 			catch (Exception ex)
 			{
@@ -532,7 +537,7 @@ namespace UnrealVS
 
 		public static bool IsGameProject(Project Project)
 		{
-			return Project.Name.EndsWith("Game", StringComparison.InvariantCultureIgnoreCase);
+			return GetUProjectNames().Any(UProject => 0 == string.Compare(UProject, Project.Name, StringComparison.OrdinalIgnoreCase));
 		}
 
 		/// <summary>
@@ -545,7 +550,29 @@ namespace UnrealVS
 
 		public static string GetUProjectFileName(Project Project)
 		{
-			return Project.Name + (".uproject");
+			return Project.Name + "." + UProjectExtension;
+		}
+
+		/// <summary>
+		/// Returns all the .uprojects found under the solution root folder.
+		/// Returns names only with no path or extension.
+		/// </summary>
+		public static IEnumerable<string> GetUProjectNames()
+		{
+			var Folder = GetSolutionFolder();
+			if (string.IsNullOrEmpty(Folder))
+			{
+				return new string[0];
+			}
+
+			if (Folder != CachedUProjectRootFolder)
+			{
+				CachedUProjectRootFolder = Folder;
+				var UProjects = Directory.GetFiles(Folder, "*." + UProjectExtension, SearchOption.AllDirectories);
+				CachedUProjectNames = (from FullPath in UProjects select Path.GetFileNameWithoutExtension(FullPath)).ToArray();
+			}
+
+			return CachedUProjectNames;
 		}
 
 		private static void PrepareOutputPane()
@@ -586,5 +613,18 @@ namespace UnrealVS
 				}
 			}
 		}
+
+		private static string GetSolutionFolder()
+		{
+			if (!UnrealVSPackage.Instance.DTE.Solution.IsOpen)
+			{
+				return string.Empty;
+			}
+
+			return Path.GetDirectoryName(UnrealVSPackage.Instance.SolutionFilepath);
+		}
+
+		private static string CachedUProjectRootFolder = string.Empty;
+		private static IEnumerable<string> CachedUProjectNames = new string[0];
 	}
 }
