UnrealVS v1.34: Added quick switching of the build config and startup project via the batch build list.

diff --git a/Engine/Source/Programs/UnrealVS/BatchBuilderToolControl.xaml b/Engine/Source/Programs/UnrealVS/BatchBuilderToolControl.xaml
--- a/Engine/Source/Programs/UnrealVS/BatchBuilderToolControl.xaml
+++ b/Engine/Source/Programs/UnrealVS/BatchBuilderToolControl.xaml
@@ -51,7 +51,9 @@
 
 		<!-- DataTemplate BuildListItemTemplate -->
 		<DataTemplate x:Key="BuildListItemTemplate">
-			<TextBlock Text="{Binding Path=DisplayString}"/>
+			<ContentControl MouseDoubleClick="OnDblClickBuildListItem" ToolTip="Double-click to switch the active solution config and the startup project to match this item.">
+				<TextBlock Text="{Binding Path=DisplayString}"/>
+			</ContentControl>
 		</DataTemplate>
 
 		<!-- DataTemplate BuildingListItemTemplate -->
diff --git a/Engine/Source/Programs/UnrealVS/BatchBuilderToolControl.xaml.cs b/Engine/Source/Programs/UnrealVS/BatchBuilderToolControl.xaml.cs
--- a/Engine/Source/Programs/UnrealVS/BatchBuilderToolControl.xaml.cs
+++ b/Engine/Source/Programs/UnrealVS/BatchBuilderToolControl.xaml.cs
@@ -571,7 +571,7 @@ namespace UnrealVS
 
 			string[] RefreshConfigs;
 			string[] RefreshPlatforms;
-			GetSolutionConfigsAndPlatforms(out RefreshConfigs, out RefreshPlatforms);
+			Utils.GetSolutionConfigsAndPlatforms(out RefreshConfigs, out RefreshPlatforms);
 
 			foreach (string Config in RefreshConfigs)
 			{
@@ -848,7 +848,7 @@ namespace UnrealVS
 			// remove invalid ones
 			string[] SolutionConfigs;
 			string[] SolutionPlatforms;
-			GetSolutionConfigsAndPlatforms(out SolutionConfigs, out SolutionPlatforms);
+			Utils.GetSolutionConfigsAndPlatforms(out SolutionConfigs, out SolutionPlatforms);
 
 			foreach (var JobSet in _BuildJobSetsCollection)
 			{
@@ -873,28 +873,6 @@ namespace UnrealVS
 			SanityCheckBuildJobs(Utils.GetAllProjectsFromDTE());
 		}
 
-		private void GetSolutionConfigsAndPlatforms(out string[] SolutionConfigs, out string[] SolutionPlatforms)
-		{
-			var UniqueConfigs = new List<string>();
-			var UniquePlatforms = new List<string>();
-
-			SolutionConfigurations DteSolutionConfigs = UnrealVSPackage.Instance.DTE.Solution.SolutionBuild.SolutionConfigurations;
-			foreach (SolutionConfiguration2 SolutionConfig in DteSolutionConfigs)
-			{
-				if (!UniqueConfigs.Contains(SolutionConfig.Name))
-				{
-					UniqueConfigs.Add(SolutionConfig.Name);
-				}
-				if (!UniquePlatforms.Contains(SolutionConfig.PlatformName))
-				{
-					UniquePlatforms.Add(SolutionConfig.PlatformName);
-				}
-			}
-
-			SolutionConfigs = UniqueConfigs.ToArray();
-			SolutionPlatforms = UniquePlatforms.ToArray();
-		}
-
 		private void TickBuildQueue()
 		{
 			if (_ActiveBuildJob != null || _BuildQueue.Count == 0) return;
@@ -988,6 +966,35 @@ namespace UnrealVS
 			DisplayBatchOutputText(Job.OutputText);
 		}
 
+		private void OnDblClickBuildListItem(object sender, MouseButtonEventArgs e)
+		{
+			var Elem = sender as FrameworkElement;
+			if (Elem == null) return;
+
+			BuildJob Job = Elem.DataContext as BuildJob;
+			if (Job == null) return;
+			
+			// Switch the Startup Project and the build config and platform to match this Job
+			if (Job.Project != null)
+			{
+				IsBusy = true;
+
+				var Project = Job.Project.GetProjectSlow();
+
+				if (Project != null)
+				{
+					// Switch to this project
+					var ProjectHierarchy = Utils.ProjectToHierarchyObject(Project);
+					UnrealVSPackage.Instance.SolutionBuildManager.set_StartupProject(ProjectHierarchy);
+
+					// Switch the build config
+					Utils.SetActiveSolutionConfiguration(Job.Config, Job.Platform);
+				}
+
+				IsBusy = false;
+			}
+		}
+
 		private void OnDblClickBuildingListItem(object sender, MouseButtonEventArgs e)
 		{
 			var Elem = sender as FrameworkElement;
diff --git a/Engine/Source/Programs/UnrealVS/Properties/AssemblyInfo.cs b/Engine/Source/Programs/UnrealVS/Properties/AssemblyInfo.cs
--- a/Engine/Source/Programs/UnrealVS/Properties/AssemblyInfo.cs
+++ b/Engine/Source/Programs/UnrealVS/Properties/AssemblyInfo.cs
@@ -31,8 +31,8 @@ using System.Runtime.InteropServices;
 // You can specify all the values or you can default the Revision and Build Numbers 
 // by using the '*' as shown below:
 
-[assembly: AssemblyVersion("1.33.0.0")]
-[assembly: AssemblyFileVersion("1.33.0.0")]
+[assembly: AssemblyVersion("1.34.0.0")]
+[assembly: AssemblyFileVersion("1.34.0.0")]
 
 
 
diff --git a/Engine/Source/Programs/UnrealVS/TargetManifests/VS11.source.extension.vsixmanifest b/Engine/Source/Programs/UnrealVS/TargetManifests/VS11.source.extension.vsixmanifest
--- a/Engine/Source/Programs/UnrealVS/TargetManifests/VS11.source.extension.vsixmanifest
+++ b/Engine/Source/Programs/UnrealVS/TargetManifests/VS11.source.extension.vsixmanifest
@@ -1,7 +1,7 @@
 ﻿<?xml version="1.0" encoding="utf-8"?>
 <PackageManifest Version="2.0.0" xmlns="http://schemas.microsoft.com/developer/vsx-schema/2011" xmlns:d="http://schemas.microsoft.com/developer/vsx-schema-design/2011">
 	<Metadata>
-		<Identity Id="ddbf523f-7eb6-4887-bd51-85a714ff87eb" Version="1.33" Language="en-US" Publisher="Epic Games, Inc." />
+		<Identity Id="ddbf523f-7eb6-4887-bd51-85a714ff87eb" Version="1.34" Language="en-US" Publisher="Epic Games, Inc." />
 		<DisplayName>UnrealVS</DisplayName>
 		<Description xml:space="preserve">Unreal Engine extension for Visual Studio</Description>
 		<MoreInfo>http://epicgames.com</MoreInfo>
diff --git a/Engine/Source/Programs/UnrealVS/TargetManifests/VS11d.source.extension.vsixmanifest b/Engine/Source/Programs/UnrealVS/TargetManifests/VS11d.source.extension.vsixmanifest
--- a/Engine/Source/Programs/UnrealVS/TargetManifests/VS11d.source.extension.vsixmanifest
+++ b/Engine/Source/Programs/UnrealVS/TargetManifests/VS11d.source.extension.vsixmanifest
@@ -1,7 +1,7 @@
 ﻿<?xml version="1.0" encoding="utf-8"?>
 <PackageManifest Version="2.0.0" xmlns="http://schemas.microsoft.com/developer/vsx-schema/2011" xmlns:d="http://schemas.microsoft.com/developer/vsx-schema-design/2011">
 	<Metadata>
-		<Identity Id="ddbf523f-7eb6-4887-bd51-85a714ff87eb" Version="1.33" Language="en-US" Publisher="Epic Games, Inc." />
+		<Identity Id="ddbf523f-7eb6-4887-bd51-85a714ff87eb" Version="1.34" Language="en-US" Publisher="Epic Games, Inc." />
 		<DisplayName>UnrealVS</DisplayName>
 		<Description xml:space="preserve">Unreal Engine extension for Visual Studio</Description>
 		<MoreInfo>http://epicgames.com</MoreInfo>
diff --git a/Engine/Source/Programs/UnrealVS/TargetManifests/VS12.source.extension.vsixmanifest b/Engine/Source/Programs/UnrealVS/TargetManifests/VS12.source.extension.vsixmanifest
--- a/Engine/Source/Programs/UnrealVS/TargetManifests/VS12.source.extension.vsixmanifest
+++ b/Engine/Source/Programs/UnrealVS/TargetManifests/VS12.source.extension.vsixmanifest
@@ -1,7 +1,7 @@
 ﻿<?xml version="1.0" encoding="utf-8"?>
 <PackageManifest Version="2.0.0" xmlns="http://schemas.microsoft.com/developer/vsx-schema/2011" xmlns:d="http://schemas.microsoft.com/developer/vsx-schema-design/2011">
 	<Metadata>
-		<Identity Id="ddbf523f-7eb6-4887-bd51-85a714ff87eb" Version="1.33" Language="en-US" Publisher="Epic Games, Inc." />
+		<Identity Id="ddbf523f-7eb6-4887-bd51-85a714ff87eb" Version="1.34" Language="en-US" Publisher="Epic Games, Inc." />
 		<DisplayName>UnrealVS</DisplayName>
 		<Description xml:space="preserve">Unreal Engine extension for Visual Studio</Description>
 		<MoreInfo>http://epicgames.com</MoreInfo>
diff --git a/Engine/Source/Programs/UnrealVS/TargetManifests/VS12d.source.extension.vsixmanifest b/Engine/Source/Programs/UnrealVS/TargetManifests/VS12d.source.extension.vsixmanifest
--- a/Engine/Source/Programs/UnrealVS/TargetManifests/VS12d.source.extension.vsixmanifest
+++ b/Engine/Source/Programs/UnrealVS/TargetManifests/VS12d.source.extension.vsixmanifest
@@ -1,7 +1,7 @@
 ﻿<?xml version="1.0" encoding="utf-8"?>
 <PackageManifest Version="2.0.0" xmlns="http://schemas.microsoft.com/developer/vsx-schema/2011" xmlns:d="http://schemas.microsoft.com/developer/vsx-schema-design/2011">
 	<Metadata>
-		<Identity Id="ddbf523f-7eb6-4887-bd51-85a714ff87eb" Version="1.33" Language="en-US" Publisher="Epic Games, Inc." />
+		<Identity Id="ddbf523f-7eb6-4887-bd51-85a714ff87eb" Version="1.34" Language="en-US" Publisher="Epic Games, Inc." />
 		<DisplayName>UnrealVS</DisplayName>
 		<Description xml:space="preserve">Unreal Engine extension for Visual Studio</Description>
 		<MoreInfo>http://epicgames.com</MoreInfo>
diff --git a/Engine/Source/Programs/UnrealVS/Utils.cs b/Engine/Source/Programs/UnrealVS/Utils.cs
--- a/Engine/Source/Programs/UnrealVS/Utils.cs
+++ b/Engine/Source/Programs/UnrealVS/Utils.cs
@@ -6,6 +6,7 @@ using Microsoft.VisualStudio.Shell.Interop;
 using EnvDTE;
 using System.Collections.Generic;
 using System.Linq;
+using System.Linq.Expressions;
 using VSLangProj;
 using System.IO;
 using EnvDTE80;
@@ -575,6 +576,43 @@ namespace UnrealVS
 			return CachedUProjectNames;
 		}
 
+		public static void GetSolutionConfigsAndPlatforms(out string[] SolutionConfigs, out string[] SolutionPlatforms)
+		{
+			var UniqueConfigs = new List<string>();
+			var UniquePlatforms = new List<string>();
+
+			SolutionConfigurations DteSolutionConfigs = UnrealVSPackage.Instance.DTE.Solution.SolutionBuild.SolutionConfigurations;
+			foreach (SolutionConfiguration2 SolutionConfig in DteSolutionConfigs)
+			{
+				if (!UniqueConfigs.Contains(SolutionConfig.Name))
+				{
+					UniqueConfigs.Add(SolutionConfig.Name);
+				}
+				if (!UniquePlatforms.Contains(SolutionConfig.PlatformName))
+				{
+					UniquePlatforms.Add(SolutionConfig.PlatformName);
+				}
+			}
+
+			SolutionConfigs = UniqueConfigs.ToArray();
+			SolutionPlatforms = UniquePlatforms.ToArray();
+		}
+
+		public static bool SetActiveSolutionConfiguration(string ConfigName, string PlatformName)
+		{
+			SolutionConfigurations DteSolutionConfigs = UnrealVSPackage.Instance.DTE.Solution.SolutionBuild.SolutionConfigurations;
+			foreach (SolutionConfiguration2 SolutionConfig in DteSolutionConfigs)
+			{
+				if (string.Compare(SolutionConfig.Name, ConfigName, StringComparison.Ordinal) == 0
+					&& string.Compare(SolutionConfig.PlatformName, PlatformName, StringComparison.Ordinal) == 0)
+				{
+					SolutionConfig.Activate();
+					return true;
+				}
+			}
+			return false;
+		}
+
 		private static void PrepareOutputPane()
 		{
 			UnrealVSPackage.Instance.DTE.ExecuteCommand("View.Output");
