UnrealVS v1.36 for VS2012/2013

Fixed Set Startup Project bug in which the startup project could be forgotten when set via the UnrealVS toolbar.

diff --git a/Engine/Source/Programs/UnrealVS/Properties/AssemblyInfo.cs b/Engine/Source/Programs/UnrealVS/Properties/AssemblyInfo.cs
--- a/Engine/Source/Programs/UnrealVS/Properties/AssemblyInfo.cs
+++ b/Engine/Source/Programs/UnrealVS/Properties/AssemblyInfo.cs
@@ -31,8 +31,8 @@ using System.Runtime.InteropServices;
 // You can specify all the values or you can default the Revision and Build Numbers 
 // by using the '*' as shown below:
 
-[assembly: AssemblyVersion("1.34.0.0")]
-[assembly: AssemblyFileVersion("1.34.0.0")]
+[assembly: AssemblyVersion("1.36.0.0")]
+[assembly: AssemblyFileVersion("1.36.0.0")]
 
 
 
diff --git a/Engine/Source/Programs/UnrealVS/StartupProjectSelector.cs b/Engine/Source/Programs/UnrealVS/StartupProjectSelector.cs
--- a/Engine/Source/Programs/UnrealVS/StartupProjectSelector.cs
+++ b/Engine/Source/Programs/UnrealVS/StartupProjectSelector.cs
@@ -265,9 +265,16 @@ namespace UnrealVS
 
 				if (ProjectRefMatch != null && ProjectRefMatch.Project != null)
 				{
-					// Switch to this project!
+					// Switch to this project using the standard SDK method set_StartupProject
 					var ProjectHierarchy = Utils.ProjectToHierarchyObject(ProjectRefMatch.Project);
 					UnrealVSPackage.Instance.SolutionBuildManager.set_StartupProject(ProjectHierarchy);
+
+					// Select in solution exp and set as startup project
+					// This forces the .suo to save the project change
+					if (Utils.SelectProjectInSolutionExplorer(ProjectRefMatch.Project))
+					{
+						UnrealVSPackage.Instance.DTE.ExecuteCommand("Project.SetasStartUpProject");
+					}
 				}					
 			}
 		}
diff --git a/Engine/Source/Programs/UnrealVS/TargetManifests/VS11.source.extension.vsixmanifest b/Engine/Source/Programs/UnrealVS/TargetManifests/VS11.source.extension.vsixmanifest
--- a/Engine/Source/Programs/UnrealVS/TargetManifests/VS11.source.extension.vsixmanifest
+++ b/Engine/Source/Programs/UnrealVS/TargetManifests/VS11.source.extension.vsixmanifest
@@ -1,7 +1,7 @@
 ﻿<?xml version="1.0" encoding="utf-8"?>
 <PackageManifest Version="2.0.0" xmlns="http://schemas.microsoft.com/developer/vsx-schema/2011" xmlns:d="http://schemas.microsoft.com/developer/vsx-schema-design/2011">
 	<Metadata>
-		<Identity Id="ddbf523f-7eb6-4887-bd51-85a714ff87eb" Version="1.34" Language="en-US" Publisher="Epic Games, Inc." />
+		<Identity Id="ddbf523f-7eb6-4887-bd51-85a714ff87eb" Version="1.36" Language="en-US" Publisher="Epic Games, Inc." />
 		<DisplayName>UnrealVS</DisplayName>
 		<Description xml:space="preserve">Unreal Engine extension for Visual Studio</Description>
 		<MoreInfo>http://epicgames.com</MoreInfo>
diff --git a/Engine/Source/Programs/UnrealVS/TargetManifests/VS11d.source.extension.vsixmanifest b/Engine/Source/Programs/UnrealVS/TargetManifests/VS11d.source.extension.vsixmanifest
--- a/Engine/Source/Programs/UnrealVS/TargetManifests/VS11d.source.extension.vsixmanifest
+++ b/Engine/Source/Programs/UnrealVS/TargetManifests/VS11d.source.extension.vsixmanifest
@@ -1,7 +1,7 @@
 ﻿<?xml version="1.0" encoding="utf-8"?>
 <PackageManifest Version="2.0.0" xmlns="http://schemas.microsoft.com/developer/vsx-schema/2011" xmlns:d="http://schemas.microsoft.com/developer/vsx-schema-design/2011">
 	<Metadata>
-		<Identity Id="ddbf523f-7eb6-4887-bd51-85a714ff87eb" Version="1.34" Language="en-US" Publisher="Epic Games, Inc." />
+		<Identity Id="ddbf523f-7eb6-4887-bd51-85a714ff87eb" Version="1.36" Language="en-US" Publisher="Epic Games, Inc." />
 		<DisplayName>UnrealVS</DisplayName>
 		<Description xml:space="preserve">Unreal Engine extension for Visual Studio</Description>
 		<MoreInfo>http://epicgames.com</MoreInfo>
diff --git a/Engine/Source/Programs/UnrealVS/TargetManifests/VS12.source.extension.vsixmanifest b/Engine/Source/Programs/UnrealVS/TargetManifests/VS12.source.extension.vsixmanifest
--- a/Engine/Source/Programs/UnrealVS/TargetManifests/VS12.source.extension.vsixmanifest
+++ b/Engine/Source/Programs/UnrealVS/TargetManifests/VS12.source.extension.vsixmanifest
@@ -1,7 +1,7 @@
 ﻿<?xml version="1.0" encoding="utf-8"?>
 <PackageManifest Version="2.0.0" xmlns="http://schemas.microsoft.com/developer/vsx-schema/2011" xmlns:d="http://schemas.microsoft.com/developer/vsx-schema-design/2011">
 	<Metadata>
-		<Identity Id="ddbf523f-7eb6-4887-bd51-85a714ff87eb" Version="1.34" Language="en-US" Publisher="Epic Games, Inc." />
+		<Identity Id="ddbf523f-7eb6-4887-bd51-85a714ff87eb" Version="1.36" Language="en-US" Publisher="Epic Games, Inc." />
 		<DisplayName>UnrealVS</DisplayName>
 		<Description xml:space="preserve">Unreal Engine extension for Visual Studio</Description>
 		<MoreInfo>http://epicgames.com</MoreInfo>
diff --git a/Engine/Source/Programs/UnrealVS/TargetManifests/VS12d.source.extension.vsixmanifest b/Engine/Source/Programs/UnrealVS/TargetManifests/VS12d.source.extension.vsixmanifest
--- a/Engine/Source/Programs/UnrealVS/TargetManifests/VS12d.source.extension.vsixmanifest
+++ b/Engine/Source/Programs/UnrealVS/TargetManifests/VS12d.source.extension.vsixmanifest
@@ -1,7 +1,7 @@
 ﻿<?xml version="1.0" encoding="utf-8"?>
 <PackageManifest Version="2.0.0" xmlns="http://schemas.microsoft.com/developer/vsx-schema/2011" xmlns:d="http://schemas.microsoft.com/developer/vsx-schema-design/2011">
 	<Metadata>
-		<Identity Id="ddbf523f-7eb6-4887-bd51-85a714ff87eb" Version="1.34" Language="en-US" Publisher="Epic Games, Inc." />
+		<Identity Id="ddbf523f-7eb6-4887-bd51-85a714ff87eb" Version="1.36" Language="en-US" Publisher="Epic Games, Inc." />
 		<DisplayName>UnrealVS</DisplayName>
 		<Description xml:space="preserve">Unreal Engine extension for Visual Studio</Description>
 		<MoreInfo>http://epicgames.com</MoreInfo>
diff --git a/Engine/Source/Programs/UnrealVS/UnrealVSPackage.cs b/Engine/Source/Programs/UnrealVS/UnrealVSPackage.cs
--- a/Engine/Source/Programs/UnrealVS/UnrealVSPackage.cs
+++ b/Engine/Source/Programs/UnrealVS/UnrealVSPackage.cs
@@ -40,7 +40,7 @@ namespace UnrealVS
 	[ProvideToolWindow(typeof(BatchBuilderToolWindow))]
 
 	// This attribute registers an options page for the package.
-	[ProvideOptionPage(typeof(UnrealVsOptions), ExtensionName, "General", 101, 102, true)]
+	//[ProvideOptionPage(typeof(UnrealVsOptions), ExtensionName, "General", 101, 102, true)]
 
 	/// <summary>
 	/// UnrealVSPackage implements Package abstract class.  This is the main class that is registered
@@ -55,7 +55,7 @@ namespace UnrealVS
 	{
 		/** Constants */
 
-		private const string VersionString = "v1.34";
+		private const string VersionString = "v1.36";
 		private const string UnrealSolutionFileNamePrefix = "UE4";
 		private const string ExtensionName = "UnrealVS";
 		private const string CommandLineOptionKey = ExtensionName + "CommandLineMRU";
@@ -460,26 +460,26 @@ namespace UnrealVS
 		/// <param name="stream">The stream to load the option data from.</param>
 		protected override void OnLoadOptions(string key, Stream stream)
 		{
-			try
-			{
-				if (0 == string.Compare(key, CommandLineOptionKey))
-				{
-					Logging.WriteLine("Restoring CommandLineEditor options");
-					CommandLineEditor.LoadOptions(stream);
-				}
-				else if (0 == string.Compare(key, BatchBuildSetsOptionKey))
-				{
-					Logging.WriteLine("Restoring BatchBuilder options");
-					BatchBuilder.LoadOptions(stream);
-				}
-			}
-			catch (Exception Ex)
-			{
-				// Couldn't load options
-				Exception AppEx = new ApplicationException("OnLoadOptions() failed with key " + key, Ex);
-				Logging.WriteLine(AppEx.ToString());
-				throw AppEx;
-			}
+			//try
+			//{
+			//	if (0 == string.Compare(key, CommandLineOptionKey))
+			//	{
+			//		Logging.WriteLine("Restoring CommandLineEditor options");
+			//		CommandLineEditor.LoadOptions(stream);
+			//	}
+			//	else if (0 == string.Compare(key, BatchBuildSetsOptionKey))
+			//	{
+			//		Logging.WriteLine("Restoring BatchBuilder options");
+			//		BatchBuilder.LoadOptions(stream);
+			//	}
+			//}
+			//catch (Exception Ex)
+			//{
+			//	// Couldn't load options
+			//	Exception AppEx = new ApplicationException("OnLoadOptions() failed with key " + key, Ex);
+			//	Logging.WriteLine(AppEx.ToString());
+			//	throw AppEx;
+			//}
 		}
 
 
@@ -491,26 +491,26 @@ namespace UnrealVS
 		/// <param name="stream">The stream to save the option data to.</param>
 		protected override void OnSaveOptions(string key, Stream stream)
 		{
-			try
-			{
-				if (0 == string.Compare(key, CommandLineOptionKey))
-				{
-					Logging.WriteLine("Saving CommandLineEditor options");
-					CommandLineEditor.SaveOptions(stream);
-				}
-				else if (0 == string.Compare(key, BatchBuildSetsOptionKey))
-				{
-					Logging.WriteLine("Saving BatchBuilder options");
-					BatchBuilder.SaveOptions(stream);
-				}
-			}
-			catch (Exception Ex)
-			{
-				// Couldn't save options
-				Exception AppEx = new ApplicationException("OnSaveOptions() failed with key " + key, Ex);
-				Logging.WriteLine(AppEx.ToString());
-				throw AppEx;
-			}
+			//try
+			//{
+			//	if (0 == string.Compare(key, CommandLineOptionKey))
+			//	{
+			//		Logging.WriteLine("Saving CommandLineEditor options");
+			//		CommandLineEditor.SaveOptions(stream);
+			//	}
+			//	else if (0 == string.Compare(key, BatchBuildSetsOptionKey))
+			//	{
+			//		Logging.WriteLine("Saving BatchBuilder options");
+			//		BatchBuilder.SaveOptions(stream);
+			//	}
+			//}
+			//catch (Exception Ex)
+			//{
+			//	// Couldn't save options
+			//	Exception AppEx = new ApplicationException("OnSaveOptions() failed with key " + key, Ex);
+			//	Logging.WriteLine(AppEx.ToString());
+			//	throw AppEx;
+			//}
 		}
 
 
diff --git a/Engine/Source/Programs/UnrealVS/Utils.cs b/Engine/Source/Programs/UnrealVS/Utils.cs
--- a/Engine/Source/Programs/UnrealVS/Utils.cs
+++ b/Engine/Source/Programs/UnrealVS/Utils.cs
@@ -177,9 +177,10 @@ namespace UnrealVS
 		/// </summary>
 		public class UITreeItem
 		{
+			public UIHierarchyItem Item { get; set; }
 			public UITreeItem[] Children { get; set; }
-			public string Name { get; set; }
-			public object Object { get; set; }
+			public string Name { get { return Item != null ? Item.Name : "None"; } }
+			public object Object { get { return Item != null ? Item.Object : null; } }
 		}
 
 		/// <summary>
@@ -189,8 +190,7 @@ namespace UnrealVS
 		{
 			return new UITreeItem
 			{
-				Name = "Root",
-				Object = null,
+				Item = null,
 				Children = (from UIHierarchyItem Child in Hierarchy.UIHierarchyItems select GetUIHierarchyTree(Child)).ToArray()
 			};
 		}
@@ -202,8 +202,7 @@ namespace UnrealVS
 		{
 			return new UITreeItem
 			{
-				Name = HierarchyItem.Name,
-				Object = HierarchyItem.Object,
+				Item = HierarchyItem,
 				Children = (from UIHierarchyItem Child in HierarchyItem.UIHierarchyItems select GetUIHierarchyTree(Child)).ToArray()
 			};
 		}
@@ -230,6 +229,22 @@ namespace UnrealVS
 			return Results;
 		}
 
+		public static IEnumerable<UIHierarchyItem> GetUITreeItemsByObjectType<T>(UITreeItem RootItem) where T : class
+		{
+			List<UIHierarchyItem> Results = new List<UIHierarchyItem>();
+
+			if (RootItem.Object is T)
+			{
+				Results.Add(RootItem.Item);
+			}
+			foreach (var Child in RootItem.Children)
+			{
+				Results.AddRange(GetUITreeItemsByObjectType<T>(Child));
+			}
+
+			return Results;
+		}
+
 		/// <summary>
 		/// Helper to check the properties of a project and determine whether it can be run in VS.
 		/// Projects that return true can be run in the debugger by pressing the usual Start Debugging (F5) command.
@@ -613,6 +628,29 @@ namespace UnrealVS
 			return false;
 		}
 
+		public static bool SelectProjectInSolutionExplorer(Project Project)
+		{
+			UnrealVSPackage.Instance.DTE.ExecuteCommand("View.SolutionExplorer");
+			Project.ParentProjectItem.ExpandView();
+
+			UIHierarchy SolutionExplorerHierarachy = UnrealVSPackage.Instance.DTE2.ToolWindows.SolutionExplorer;
+			Utils.UITreeItem SolutionExplorerTree = Utils.GetUIHierarchyTree(SolutionExplorerHierarachy);
+			var UIHierarachyProjects = Utils.GetUITreeItemsByObjectType<Project>(SolutionExplorerTree);
+
+			var SelectableUIItem = UIHierarachyProjects.FirstOrDefault(uihp => uihp.Object as Project == Project);
+
+			if (SelectableUIItem != null)
+			{
+				if (Project.ParentProjectItem != null)
+				{
+					SelectableUIItem.Select(vsUISelectionType.vsUISelectionTypeSelect);
+					return true;
+				}
+			}
+
+			return false;
+		}
+
 		private static void PrepareOutputPane()
 		{
 			UnrealVSPackage.Instance.DTE.ExecuteCommand("View.Output");
