UnrealVS v1.37 for VS2012/2013

Reenabled load/save UnrealVS settings to .suo

diff --git a/Engine/Source/Programs/UnrealVS/Properties/AssemblyInfo.cs b/Engine/Source/Programs/UnrealVS/Properties/AssemblyInfo.cs
--- a/Engine/Source/Programs/UnrealVS/Properties/AssemblyInfo.cs
+++ b/Engine/Source/Programs/UnrealVS/Properties/AssemblyInfo.cs
@@ -31,8 +31,8 @@ using System.Runtime.InteropServices;
 // You can specify all the values or you can default the Revision and Build Numbers 
 // by using the '*' as shown below:
 
-[assembly: AssemblyVersion("1.36.0.0")]
-[assembly: AssemblyFileVersion("1.36.0.0")]
+[assembly: AssemblyVersion("1.37.0.0")]
+[assembly: AssemblyFileVersion("1.37.0.0")]
 
 
 
diff --git a/Engine/Source/Programs/UnrealVS/TargetManifests/VS11.source.extension.vsixmanifest b/Engine/Source/Programs/UnrealVS/TargetManifests/VS11.source.extension.vsixmanifest
--- a/Engine/Source/Programs/UnrealVS/TargetManifests/VS11.source.extension.vsixmanifest
+++ b/Engine/Source/Programs/UnrealVS/TargetManifests/VS11.source.extension.vsixmanifest
@@ -1,7 +1,7 @@
 ﻿<?xml version="1.0" encoding="utf-8"?>
 <PackageManifest Version="2.0.0" xmlns="http://schemas.microsoft.com/developer/vsx-schema/2011" xmlns:d="http://schemas.microsoft.com/developer/vsx-schema-design/2011">
 	<Metadata>
-		<Identity Id="ddbf523f-7eb6-4887-bd51-85a714ff87eb" Version="1.36" Language="en-US" Publisher="Epic Games, Inc." />
+		<Identity Id="ddbf523f-7eb6-4887-bd51-85a714ff87eb" Version="1.37" Language="en-US" Publisher="Epic Games, Inc." />
 		<DisplayName>UnrealVS</DisplayName>
 		<Description xml:space="preserve">Unreal Engine extension for Visual Studio</Description>
 		<MoreInfo>http://epicgames.com</MoreInfo>
diff --git a/Engine/Source/Programs/UnrealVS/TargetManifests/VS11d.source.extension.vsixmanifest b/Engine/Source/Programs/UnrealVS/TargetManifests/VS11d.source.extension.vsixmanifest
--- a/Engine/Source/Programs/UnrealVS/TargetManifests/VS11d.source.extension.vsixmanifest
+++ b/Engine/Source/Programs/UnrealVS/TargetManifests/VS11d.source.extension.vsixmanifest
@@ -1,7 +1,7 @@
 ﻿<?xml version="1.0" encoding="utf-8"?>
 <PackageManifest Version="2.0.0" xmlns="http://schemas.microsoft.com/developer/vsx-schema/2011" xmlns:d="http://schemas.microsoft.com/developer/vsx-schema-design/2011">
 	<Metadata>
-		<Identity Id="ddbf523f-7eb6-4887-bd51-85a714ff87eb" Version="1.36" Language="en-US" Publisher="Epic Games, Inc." />
+		<Identity Id="ddbf523f-7eb6-4887-bd51-85a714ff87eb" Version="1.37" Language="en-US" Publisher="Epic Games, Inc." />
 		<DisplayName>UnrealVS</DisplayName>
 		<Description xml:space="preserve">Unreal Engine extension for Visual Studio</Description>
 		<MoreInfo>http://epicgames.com</MoreInfo>
diff --git a/Engine/Source/Programs/UnrealVS/TargetManifests/VS12.source.extension.vsixmanifest b/Engine/Source/Programs/UnrealVS/TargetManifests/VS12.source.extension.vsixmanifest
--- a/Engine/Source/Programs/UnrealVS/TargetManifests/VS12.source.extension.vsixmanifest
+++ b/Engine/Source/Programs/UnrealVS/TargetManifests/VS12.source.extension.vsixmanifest
@@ -1,7 +1,7 @@
 ﻿<?xml version="1.0" encoding="utf-8"?>
 <PackageManifest Version="2.0.0" xmlns="http://schemas.microsoft.com/developer/vsx-schema/2011" xmlns:d="http://schemas.microsoft.com/developer/vsx-schema-design/2011">
 	<Metadata>
-		<Identity Id="ddbf523f-7eb6-4887-bd51-85a714ff87eb" Version="1.36" Language="en-US" Publisher="Epic Games, Inc." />
+		<Identity Id="ddbf523f-7eb6-4887-bd51-85a714ff87eb" Version="1.37" Language="en-US" Publisher="Epic Games, Inc." />
 		<DisplayName>UnrealVS</DisplayName>
 		<Description xml:space="preserve">Unreal Engine extension for Visual Studio</Description>
 		<MoreInfo>http://epicgames.com</MoreInfo>
diff --git a/Engine/Source/Programs/UnrealVS/TargetManifests/VS12d.source.extension.vsixmanifest b/Engine/Source/Programs/UnrealVS/TargetManifests/VS12d.source.extension.vsixmanifest
--- a/Engine/Source/Programs/UnrealVS/TargetManifests/VS12d.source.extension.vsixmanifest
+++ b/Engine/Source/Programs/UnrealVS/TargetManifests/VS12d.source.extension.vsixmanifest
@@ -1,7 +1,7 @@
 ﻿<?xml version="1.0" encoding="utf-8"?>
 <PackageManifest Version="2.0.0" xmlns="http://schemas.microsoft.com/developer/vsx-schema/2011" xmlns:d="http://schemas.microsoft.com/developer/vsx-schema-design/2011">
 	<Metadata>
-		<Identity Id="ddbf523f-7eb6-4887-bd51-85a714ff87eb" Version="1.36" Language="en-US" Publisher="Epic Games, Inc." />
+		<Identity Id="ddbf523f-7eb6-4887-bd51-85a714ff87eb" Version="1.37" Language="en-US" Publisher="Epic Games, Inc." />
 		<DisplayName>UnrealVS</DisplayName>
 		<Description xml:space="preserve">Unreal Engine extension for Visual Studio</Description>
 		<MoreInfo>http://epicgames.com</MoreInfo>
diff --git a/Engine/Source/Programs/UnrealVS/UnrealVSPackage.cs b/Engine/Source/Programs/UnrealVS/UnrealVSPackage.cs
--- a/Engine/Source/Programs/UnrealVS/UnrealVSPackage.cs
+++ b/Engine/Source/Programs/UnrealVS/UnrealVSPackage.cs
@@ -40,7 +40,7 @@ namespace UnrealVS
 	[ProvideToolWindow(typeof(BatchBuilderToolWindow))]
 
 	// This attribute registers an options page for the package.
-	//[ProvideOptionPage(typeof(UnrealVsOptions), ExtensionName, "General", 101, 102, true)]
+	[ProvideOptionPage(typeof(UnrealVsOptions), ExtensionName, "General", 101, 102, true)]
 
 	/// <summary>
 	/// UnrealVSPackage implements Package abstract class.  This is the main class that is registered
@@ -55,7 +55,7 @@ namespace UnrealVS
 	{
 		/** Constants */
 
-		private const string VersionString = "v1.36";
+		private const string VersionString = "v1.37";
 		private const string UnrealSolutionFileNamePrefix = "UE4";
 		private const string ExtensionName = "UnrealVS";
 		private const string CommandLineOptionKey = ExtensionName + "CommandLineMRU";
@@ -460,26 +460,26 @@ namespace UnrealVS
 		/// <param name="stream">The stream to load the option data from.</param>
 		protected override void OnLoadOptions(string key, Stream stream)
 		{
-			//try
-			//{
-			//	if (0 == string.Compare(key, CommandLineOptionKey))
-			//	{
-			//		Logging.WriteLine("Restoring CommandLineEditor options");
-			//		CommandLineEditor.LoadOptions(stream);
-			//	}
-			//	else if (0 == string.Compare(key, BatchBuildSetsOptionKey))
-			//	{
-			//		Logging.WriteLine("Restoring BatchBuilder options");
-			//		BatchBuilder.LoadOptions(stream);
-			//	}
-			//}
-			//catch (Exception Ex)
-			//{
-			//	// Couldn't load options
-			//	Exception AppEx = new ApplicationException("OnLoadOptions() failed with key " + key, Ex);
-			//	Logging.WriteLine(AppEx.ToString());
-			//	throw AppEx;
-			//}
+			try
+			{
+				if (0 == string.Compare(key, CommandLineOptionKey))
+				{
+					Logging.WriteLine("Restoring CommandLineEditor options");
+					CommandLineEditor.LoadOptions(stream);
+				}
+				else if (0 == string.Compare(key, BatchBuildSetsOptionKey))
+				{
+					Logging.WriteLine("Restoring BatchBuilder options");
+					BatchBuilder.LoadOptions(stream);
+				}
+			}
+			catch (Exception Ex)
+			{
+				// Couldn't load options
+				Exception AppEx = new ApplicationException("OnLoadOptions() failed with key " + key, Ex);
+				Logging.WriteLine(AppEx.ToString());
+				throw AppEx;
+			}
 		}
 
 
@@ -491,26 +491,26 @@ namespace UnrealVS
 		/// <param name="stream">The stream to save the option data to.</param>
 		protected override void OnSaveOptions(string key, Stream stream)
 		{
-			//try
-			//{
-			//	if (0 == string.Compare(key, CommandLineOptionKey))
-			//	{
-			//		Logging.WriteLine("Saving CommandLineEditor options");
-			//		CommandLineEditor.SaveOptions(stream);
-			//	}
-			//	else if (0 == string.Compare(key, BatchBuildSetsOptionKey))
-			//	{
-			//		Logging.WriteLine("Saving BatchBuilder options");
-			//		BatchBuilder.SaveOptions(stream);
-			//	}
-			//}
-			//catch (Exception Ex)
-			//{
-			//	// Couldn't save options
-			//	Exception AppEx = new ApplicationException("OnSaveOptions() failed with key " + key, Ex);
-			//	Logging.WriteLine(AppEx.ToString());
-			//	throw AppEx;
-			//}
+			try
+			{
+				if (0 == string.Compare(key, CommandLineOptionKey))
+				{
+					Logging.WriteLine("Saving CommandLineEditor options");
+					CommandLineEditor.SaveOptions(stream);
+				}
+				else if (0 == string.Compare(key, BatchBuildSetsOptionKey))
+				{
+					Logging.WriteLine("Saving BatchBuilder options");
+					BatchBuilder.SaveOptions(stream);
+				}
+			}
+			catch (Exception Ex)
+			{
+				// Couldn't save options
+				Exception AppEx = new ApplicationException("OnSaveOptions() failed with key " + key, Ex);
+				Logging.WriteLine(AppEx.ToString());
+				throw AppEx;
+			}
 		}
 
 
