Improved handling of command line arguments in UnrealVS - stops bad uproject paths in user created projects (UnrealVS 1.38)
[UE-12222] - UnrealVS automatically adds a bad .uproject path to user created projects

diff --git a/Engine/Source/Programs/UnrealVS/CommandLineEditor.cs b/Engine/Source/Programs/UnrealVS/CommandLineEditor.cs
--- a/Engine/Source/Programs/UnrealVS/CommandLineEditor.cs
+++ b/Engine/Source/Programs/UnrealVS/CommandLineEditor.cs
@@ -161,14 +161,18 @@ namespace UnrealVS
 						var ActiveConfiguration = (SolutionConfiguration2) UnrealVSPackage.Instance.DTE.Solution.SolutionBuild.ActiveConfiguration;
 						if (UnrealVSPackage.Instance.IsUE4Loaded && Utils.IsGameProject(SelectedStartupProject) && Utils.HasUProjectCommandLineArg(ActiveConfiguration.Name))
 						{
-							string UProjectFileName = Utils.GetUProjectFileName(SelectedStartupProject);
-							if (CommandLineArguments.Trim().StartsWith(UProjectFileName, StringComparison.OrdinalIgnoreCase))
-							{
-								CommandLineArguments = CommandLineArguments.Trim().Substring(UProjectFileName.Length).Trim();
-							}
-							else if (CommandLineArguments.Trim().StartsWith(SelectedStartupProject.Name, StringComparison.OrdinalIgnoreCase))
+							var AutoPrefix = Utils.GetAutoUProjectCommandLinePrefix(SelectedStartupProject);
+
+							if (!string.IsNullOrEmpty(AutoPrefix))
 							{
-								CommandLineArguments = CommandLineArguments.Trim().Substring(SelectedStartupProject.Name.Length).Trim();
+								if (CommandLineArguments.Trim().StartsWith(AutoPrefix, StringComparison.OrdinalIgnoreCase))
+								{
+									CommandLineArguments = CommandLineArguments.Trim().Substring(AutoPrefix.Length).Trim();
+								}
+								//else if (CommandLineArguments.Trim().StartsWith(SelectedStartupProject.Name + " ", StringComparison.OrdinalIgnoreCase))
+								//{
+								//	CommandLineArguments = CommandLineArguments.Trim().Substring(SelectedStartupProject.Name.Length + 1).Trim();
+								//}
 							}
 						}
 
@@ -245,9 +249,6 @@ namespace UnrealVS
 
 		private void CommitCommandLineText(string CommandLine)
 		{
-			
-			string FullCommandLine = CommandLine;
-
 			IVsHierarchy ProjectHierarchy;
 			UnrealVSPackage.Instance.SolutionBuildManager.get_StartupProject(out ProjectHierarchy);
 			if (ProjectHierarchy != null)
@@ -256,39 +257,27 @@ namespace UnrealVS
 
 				if (SelectedStartupProject != null)
 				{
-					// for "Game" projects automatically remove the game project filename from the start of the command line
-					var ActiveConfiguration = (SolutionConfiguration2)UnrealVSPackage.Instance.DTE.Solution.SolutionBuild.ActiveConfiguration;
-					if (UnrealVSPackage.Instance.IsUE4Loaded && Utils.IsGameProject(SelectedStartupProject) && Utils.HasUProjectCommandLineArg(ActiveConfiguration.Name))
-					{
-						string UProjectFileName = Utils.GetUProjectFileName(SelectedStartupProject);
-						if (FullCommandLine.Trim().StartsWith(UProjectFileName, StringComparison.InvariantCultureIgnoreCase))
-						{
-							VsShellUtilities.ShowMessageBox(ServiceProvider.GlobalProvider,
-															string.Format("INFORMATION: The project filename {0} has been removed from the command line because it is included automatically for 'Game' projects.", UProjectFileName),
-															"UnrealVS",
-															OLEMSGICON.OLEMSGICON_INFO,
-															OLEMSGBUTTON.OLEMSGBUTTON_OK,
-															OLEMSGDEFBUTTON.OLEMSGDEFBUTTON_FIRST);
-						}
-						else if (FullCommandLine.Trim().StartsWith(SelectedStartupProject.Name, StringComparison.OrdinalIgnoreCase))
-						{
-							VsShellUtilities.ShowMessageBox(ServiceProvider.GlobalProvider,
-															string.Format("INFORMATION: The project name {0} has been removed from the command line because it is included automatically for 'Game' projects.", SelectedStartupProject.Name),
-															"UnrealVS",
-															OLEMSGICON.OLEMSGICON_INFO,
-															OLEMSGBUTTON.OLEMSGBUTTON_OK,
-															OLEMSGDEFBUTTON.OLEMSGDEFBUTTON_FIRST);
-						}
-						else
-						{
-							FullCommandLine = UProjectFileName + " " + FullCommandLine;
-						}
-					}
-
 					var CommandLineArgumentsProperty = GetProjectCommandLineProperty(SelectedStartupProject);
 
 					if (CommandLineArgumentsProperty != null)
 					{
+						string FullCommandLine = CommandLine;
+
+						// for "Game" projects automatically remove the game project filename from the start of the command line
+						var ActiveConfiguration = (SolutionConfiguration2)UnrealVSPackage.Instance.DTE.Solution.SolutionBuild.ActiveConfiguration;
+						if (UnrealVSPackage.Instance.IsUE4Loaded && Utils.IsGameProject(SelectedStartupProject) && Utils.HasUProjectCommandLineArg(ActiveConfiguration.Name))
+						{
+							var AutoPrefix = Utils.GetAutoUProjectCommandLinePrefix(SelectedStartupProject);
+							
+							if (FullCommandLine.IndexOf(Utils.UProjectExtension, StringComparison.OrdinalIgnoreCase) < 0 &&
+								string.Compare(FullCommandLine, SelectedStartupProject.Name, StringComparison.OrdinalIgnoreCase) != 0 &&
+								!FullCommandLine.StartsWith(SelectedStartupProject.Name + " ", StringComparison.OrdinalIgnoreCase))
+							{
+								// Committed command line does not specify a .uproject
+								FullCommandLine = AutoPrefix + " " + FullCommandLine;
+							}
+						}
+
 						Utils.SetPropertyValue(CommandLineArgumentsProperty, FullCommandLine);
 					}
 				}
diff --git a/Engine/Source/Programs/UnrealVS/Properties/AssemblyInfo.cs b/Engine/Source/Programs/UnrealVS/Properties/AssemblyInfo.cs
--- a/Engine/Source/Programs/UnrealVS/Properties/AssemblyInfo.cs
+++ b/Engine/Source/Programs/UnrealVS/Properties/AssemblyInfo.cs
@@ -31,8 +31,8 @@ using System.Runtime.InteropServices;
 // You can specify all the values or you can default the Revision and Build Numbers 
 // by using the '*' as shown below:
 
-[assembly: AssemblyVersion("1.37.0.0")]
-[assembly: AssemblyFileVersion("1.37.0.0")]
+[assembly: AssemblyVersion("1.38.0.0")]
+[assembly: AssemblyFileVersion("1.38.0.0")]
 
 
 
diff --git a/Engine/Source/Programs/UnrealVS/TargetManifests/VS11.source.extension.vsixmanifest b/Engine/Source/Programs/UnrealVS/TargetManifests/VS11.source.extension.vsixmanifest
--- a/Engine/Source/Programs/UnrealVS/TargetManifests/VS11.source.extension.vsixmanifest
+++ b/Engine/Source/Programs/UnrealVS/TargetManifests/VS11.source.extension.vsixmanifest
@@ -1,7 +1,7 @@
 ﻿<?xml version="1.0" encoding="utf-8"?>
 <PackageManifest Version="2.0.0" xmlns="http://schemas.microsoft.com/developer/vsx-schema/2011" xmlns:d="http://schemas.microsoft.com/developer/vsx-schema-design/2011">
 	<Metadata>
-		<Identity Id="ddbf523f-7eb6-4887-bd51-85a714ff87eb" Version="1.37" Language="en-US" Publisher="Epic Games, Inc." />
+		<Identity Id="ddbf523f-7eb6-4887-bd51-85a714ff87eb" Version="1.38" Language="en-US" Publisher="Epic Games, Inc." />
 		<DisplayName>UnrealVS</DisplayName>
 		<Description xml:space="preserve">Unreal Engine extension for Visual Studio</Description>
 		<MoreInfo>http://epicgames.com</MoreInfo>
diff --git a/Engine/Source/Programs/UnrealVS/TargetManifests/VS11d.source.extension.vsixmanifest b/Engine/Source/Programs/UnrealVS/TargetManifests/VS11d.source.extension.vsixmanifest
--- a/Engine/Source/Programs/UnrealVS/TargetManifests/VS11d.source.extension.vsixmanifest
+++ b/Engine/Source/Programs/UnrealVS/TargetManifests/VS11d.source.extension.vsixmanifest
@@ -1,7 +1,7 @@
 ﻿<?xml version="1.0" encoding="utf-8"?>
 <PackageManifest Version="2.0.0" xmlns="http://schemas.microsoft.com/developer/vsx-schema/2011" xmlns:d="http://schemas.microsoft.com/developer/vsx-schema-design/2011">
 	<Metadata>
-		<Identity Id="ddbf523f-7eb6-4887-bd51-85a714ff87eb" Version="1.37" Language="en-US" Publisher="Epic Games, Inc." />
+		<Identity Id="ddbf523f-7eb6-4887-bd51-85a714ff87eb" Version="1.38" Language="en-US" Publisher="Epic Games, Inc." />
 		<DisplayName>UnrealVS</DisplayName>
 		<Description xml:space="preserve">Unreal Engine extension for Visual Studio</Description>
 		<MoreInfo>http://epicgames.com</MoreInfo>
diff --git a/Engine/Source/Programs/UnrealVS/TargetManifests/VS12.source.extension.vsixmanifest b/Engine/Source/Programs/UnrealVS/TargetManifests/VS12.source.extension.vsixmanifest
--- a/Engine/Source/Programs/UnrealVS/TargetManifests/VS12.source.extension.vsixmanifest
+++ b/Engine/Source/Programs/UnrealVS/TargetManifests/VS12.source.extension.vsixmanifest
@@ -1,7 +1,7 @@
 ﻿<?xml version="1.0" encoding="utf-8"?>
 <PackageManifest Version="2.0.0" xmlns="http://schemas.microsoft.com/developer/vsx-schema/2011" xmlns:d="http://schemas.microsoft.com/developer/vsx-schema-design/2011">
 	<Metadata>
-		<Identity Id="ddbf523f-7eb6-4887-bd51-85a714ff87eb" Version="1.37" Language="en-US" Publisher="Epic Games, Inc." />
+		<Identity Id="ddbf523f-7eb6-4887-bd51-85a714ff87eb" Version="1.38" Language="en-US" Publisher="Epic Games, Inc." />
 		<DisplayName>UnrealVS</DisplayName>
 		<Description xml:space="preserve">Unreal Engine extension for Visual Studio</Description>
 		<MoreInfo>http://epicgames.com</MoreInfo>
diff --git a/Engine/Source/Programs/UnrealVS/TargetManifests/VS12d.source.extension.vsixmanifest b/Engine/Source/Programs/UnrealVS/TargetManifests/VS12d.source.extension.vsixmanifest
--- a/Engine/Source/Programs/UnrealVS/TargetManifests/VS12d.source.extension.vsixmanifest
+++ b/Engine/Source/Programs/UnrealVS/TargetManifests/VS12d.source.extension.vsixmanifest
@@ -1,7 +1,7 @@
 ﻿<?xml version="1.0" encoding="utf-8"?>
 <PackageManifest Version="2.0.0" xmlns="http://schemas.microsoft.com/developer/vsx-schema/2011" xmlns:d="http://schemas.microsoft.com/developer/vsx-schema-design/2011">
 	<Metadata>
-		<Identity Id="ddbf523f-7eb6-4887-bd51-85a714ff87eb" Version="1.37" Language="en-US" Publisher="Epic Games, Inc." />
+		<Identity Id="ddbf523f-7eb6-4887-bd51-85a714ff87eb" Version="1.38" Language="en-US" Publisher="Epic Games, Inc." />
 		<DisplayName>UnrealVS</DisplayName>
 		<Description xml:space="preserve">Unreal Engine extension for Visual Studio</Description>
 		<MoreInfo>http://epicgames.com</MoreInfo>
diff --git a/Engine/Source/Programs/UnrealVS/UnrealVSPackage.cs b/Engine/Source/Programs/UnrealVS/UnrealVSPackage.cs
--- a/Engine/Source/Programs/UnrealVS/UnrealVSPackage.cs
+++ b/Engine/Source/Programs/UnrealVS/UnrealVSPackage.cs
@@ -55,7 +55,7 @@ namespace UnrealVS
 	{
 		/** Constants */
 
-		private const string VersionString = "v1.37";
+		private const string VersionString = "v1.38";
 		private const string UnrealSolutionFileNamePrefix = "UE4";
 		private const string ExtensionName = "UnrealVS";
 		private const string CommandLineOptionKey = ExtensionName + "CommandLineMRU";
@@ -451,6 +451,10 @@ namespace UnrealVS
 			return GetOutputPane(VSConstants.OutputWindowPaneGuid.BuildOutputPane_guid, "Build");
 		}
 
+		public IEnumerable<string> GetLoadedProjectPaths()
+		{
+			return LoadedProjectPaths;
+		}
 
 		/// <summary>
 		/// Overrides Package.OnLoadOptions()
@@ -541,8 +545,10 @@ namespace UnrealVS
 
 			// Get the actual Project object from the IVsHierarchy object that was supplied
 			var OpenedProject = Utils.HierarchyObjectToProject( pHierarchy );
+			Utils.OnProjectListChanged();
 			if (OpenedProject != null && OnProjectOpened != null)
-            {
+			{
+				LoadedProjectPaths.Add(OpenedProject.FullName);
                 OnProjectOpened(OpenedProject);
             }
 
@@ -571,6 +577,7 @@ namespace UnrealVS
 			var ClosedProject = Utils.HierarchyObjectToProject( pHierarchy );
 			if (ClosedProject != null && OnProjectClosed != null)
             {
+				LoadedProjectPaths.Remove(ClosedProject.FullName);
                 OnProjectClosed(ClosedProject);
             }
 
@@ -837,6 +844,8 @@ namespace UnrealVS
 
 		private string _UBTVersion = string.Empty;
 
+		private readonly List<string> LoadedProjectPaths = new List<string>(); 
+
 		/// Obtains the DTE2 interface for this instance of VS from the RunningObjectTable
 		private static DTE2 GetDTE2ForCurrentInstance(DTE DTE)
 		{
diff --git a/Engine/Source/Programs/UnrealVS/Utils.cs b/Engine/Source/Programs/UnrealVS/Utils.cs
--- a/Engine/Source/Programs/UnrealVS/Utils.cs
+++ b/Engine/Source/Programs/UnrealVS/Utils.cs
@@ -553,7 +553,7 @@ namespace UnrealVS
 
 		public static bool IsGameProject(Project Project)
 		{
-			return GetUProjectNames().Any(UProject => 0 == string.Compare(UProject, Project.Name, StringComparison.OrdinalIgnoreCase));
+			return GetUProjects().ContainsKey(Project.Name);
 		}
 
 		/// <summary>
@@ -569,26 +569,74 @@ namespace UnrealVS
 			return Project.Name + "." + UProjectExtension;
 		}
 
+		public static string GetAutoUProjectCommandLinePrefix(Project Project)
+		{
+			var UProjectFileName = GetUProjectFileName(Project);
+			var AllUProjects = GetUProjects();
+
+			string UProjectPath = string.Empty;
+			if (!AllUProjects.TryGetValue(Project.Name, out UProjectPath))
+			{
+				// Search the project folder
+				var ProjectFolder = Path.GetDirectoryName(Project.FullName);
+				var UProjUnderProject = Directory.GetFiles(ProjectFolder, UProjectFileName, SearchOption.TopDirectoryOnly);
+				if (UProjUnderProject.Length == 1)
+				{
+					UProjectPath = UProjUnderProject[0];
+				}				
+			}
+
+			return '\"' + UProjectPath + '\"';
+		}
+
 		/// <summary>
 		/// Returns all the .uprojects found under the solution root folder.
-		/// Returns names only with no path or extension.
 		/// </summary>
-		public static IEnumerable<string> GetUProjectNames()
+		public static IDictionary<string, string> GetUProjects()
 		{
 			var Folder = GetSolutionFolder();
 			if (string.IsNullOrEmpty(Folder))
 			{
-				return new string[0];
+				return new Dictionary<string, string>();
 			}
 
 			if (Folder != CachedUProjectRootFolder)
 			{
+				Logging.WriteLine("GetUProjects: recaching uproject paths...");
+
 				CachedUProjectRootFolder = Folder;
-				var UProjects = Directory.GetFiles(Folder, "*." + UProjectExtension, SearchOption.AllDirectories);
-				CachedUProjectNames = (from FullPath in UProjects select Path.GetFileNameWithoutExtension(FullPath)).ToArray();
+				CachedUProjectPaths = Directory.GetFiles(Folder, "*." + UProjectExtension, SearchOption.AllDirectories);
+				CachedUProjects = null;
+
+				Logging.WriteLine("    DONE");
 			}
 
-			return CachedUProjectNames;
+			if (CachedUProjects == null)
+			{
+				Logging.WriteLine("GetUProjects: recaching uproject names...");
+
+				var ProjectPaths = UnrealVSPackage.Instance.GetLoadedProjectPaths();
+				var ProjectNames = (from path in ProjectPaths select Path.GetFileNameWithoutExtension(path)).ToArray();
+
+				var CodeUProjects = from UProjectPath in CachedUProjectPaths
+					let ProjectName = Path.GetFileNameWithoutExtension(UProjectPath)
+					where ProjectNames.Any(name => string.Compare(name, ProjectName, StringComparison.OrdinalIgnoreCase) == 0)
+					select new {Name = ProjectName, FilePath = UProjectPath};
+
+				CachedUProjects = new Dictionary<string, string>();
+
+				foreach (var UProject in CodeUProjects)
+				{
+					if (!CachedUProjects.ContainsKey(UProject.Name))
+					{
+						CachedUProjects.Add(UProject.Name, UProject.FilePath);
+					}
+				}
+
+				Logging.WriteLine("    DONE");
+			}
+
+			return CachedUProjects;
 		}
 
 		public static void GetSolutionConfigsAndPlatforms(out string[] SolutionConfigs, out string[] SolutionPlatforms)
@@ -651,6 +699,11 @@ namespace UnrealVS
 			return false;
 		}
 
+		public static void OnProjectListChanged()
+		{
+			CachedUProjects = null;
+		}
+
 		private static void PrepareOutputPane()
 		{
 			UnrealVSPackage.Instance.DTE.ExecuteCommand("View.Output");
@@ -701,6 +754,7 @@ namespace UnrealVS
 		}
 
 		private static string CachedUProjectRootFolder = string.Empty;
-		private static IEnumerable<string> CachedUProjectNames = new string[0];
+		private static IEnumerable<string> CachedUProjectPaths = new string[0];
+		private static IDictionary<string, string> CachedUProjects = null;
 	}
 }
