UnrealVS 1.39 - coverted solution to VS2015
[UE-19378] - UnrealVS update to VS2015 complete

diff --git a/Engine/Source/Programs/UnrealVS/Properties/AssemblyInfo.cs b/Engine/Source/Programs/UnrealVS/Properties/AssemblyInfo.cs
--- a/Engine/Source/Programs/UnrealVS/Properties/AssemblyInfo.cs
+++ b/Engine/Source/Programs/UnrealVS/Properties/AssemblyInfo.cs
@@ -31,8 +31,8 @@ using System.Runtime.InteropServices;
 // You can specify all the values or you can default the Revision and Build Numbers 
 // by using the '*' as shown below:
 
-[assembly: AssemblyVersion("1.38.0.0")]
-[assembly: AssemblyFileVersion("1.38.0.0")]
+[assembly: AssemblyVersion("1.39.0.0")]
+[assembly: AssemblyFileVersion("1.39.0.0")]
 
 
 
diff --git a/Engine/Source/Programs/UnrealVS/UnrealVS.csproj b/Engine/Source/Programs/UnrealVS/UnrealVS.csproj
--- a/Engine/Source/Programs/UnrealVS/UnrealVS.csproj
+++ b/Engine/Source/Programs/UnrealVS/UnrealVS.csproj
@@ -1,5 +1,5 @@
 ﻿<?xml version="1.0" encoding="utf-8"?>
-<Project DefaultTargets="Build" xmlns="http://schemas.microsoft.com/developer/msbuild/2003" ToolsVersion="12.0">
+<Project DefaultTargets="Build" xmlns="http://schemas.microsoft.com/developer/msbuild/2003" ToolsVersion="14.0">
   <Import Project="$(MSBuildExtensionsPath)\$(MSBuildToolsVersion)\Microsoft.Common.props" Condition="Exists('$(MSBuildExtensionsPath)\$(MSBuildToolsVersion)\Microsoft.Common.props')" />
   <PropertyGroup>
     <Configuration Condition=" '$(Configuration)' == '' ">Debug</Configuration>
@@ -22,8 +22,8 @@
     </SccAuxPath>
     <SccProvider>
     </SccProvider>
-    <MinimumVisualStudioVersion>12.0</MinimumVisualStudioVersion>
-    <OldToolsVersion>4.0</OldToolsVersion>
+    <MinimumVisualStudioVersion>14.0</MinimumVisualStudioVersion>
+    <OldToolsVersion>12.0</OldToolsVersion>
     <TargetFrameworkProfile />
     <FileUpgradeFlags>
     </FileUpgradeFlags>
@@ -85,6 +85,7 @@
     <Reference Include="Microsoft.VisualStudio.OLE.Interop, Version=7.1.40304.0, Culture=neutral, PublicKeyToken=b03f5f7f11d50a3a" />
     <Reference Include="Microsoft.VisualStudio.Shell.11.0, Version=11.0.0.0, Culture=neutral, PublicKeyToken=b03f5f7f11d50a3a, processorArchitecture=MSIL" Condition="'$(Configuration)' == 'VS11' Or '$(Configuration)' == 'VS11d'" />
     <Reference Include="Microsoft.VisualStudio.Shell.12.0, Version=12.0.0.0, Culture=neutral, PublicKeyToken=b03f5f7f11d50a3a, processorArchitecture=MSIL" Condition="'$(Configuration)' == 'VS12' Or '$(Configuration)' == 'VS12d'" />
+    <Reference Include="Microsoft.VisualStudio.Shell.14.0, Version=14.0.0.0, Culture=neutral, PublicKeyToken=b03f5f7f11d50a3a, processorArchitecture=MSIL" Condition="'$(Configuration)' == 'VS14' Or '$(Configuration)' == 'VS14d'" />
     <Reference Include="Microsoft.VisualStudio.Shell.Immutable.10.0" />
     <Reference Include="Microsoft.VisualStudio.Shell.Interop, Version=7.1.40304.0, Culture=neutral, PublicKeyToken=b03f5f7f11d50a3a" />
     <Reference Include="Microsoft.VisualStudio.Shell.Interop.10.0, Version=10.0.0.0, Culture=neutral, PublicKeyToken=b03f5f7f11d50a3a, processorArchitecture=MSIL">
@@ -96,12 +97,18 @@
     <Reference Include="Microsoft.VisualStudio.Shell.Interop.8.0, Version=8.0.0.0, Culture=neutral, PublicKeyToken=b03f5f7f11d50a3a" />
     <Reference Include="Microsoft.VisualStudio.Shell.Interop.9.0, Version=9.0.0.0, Culture=neutral, PublicKeyToken=b03f5f7f11d50a3a" />
     <Reference Include="Microsoft.VisualStudio.TextManager.Interop, Version=7.1.40304.0, Culture=neutral, PublicKeyToken=b03f5f7f11d50a3a" />
+    <Reference Include="Microsoft.VisualStudio.VCProject, Version=14.0.0.0, Culture=neutral, PublicKeyToken=b03f5f7f11d50a3a" Condition="'$(Configuration)' == 'VS14' Or '$(Configuration)' == 'VS14d'">
+      <EmbedInteropTypes>True</EmbedInteropTypes>
+    </Reference>
     <Reference Include="Microsoft.VisualStudio.VCProject, Version=12.0.0.0, Culture=neutral, PublicKeyToken=b03f5f7f11d50a3a" Condition="'$(Configuration)' == 'VS12' Or '$(Configuration)' == 'VS12d'">
       <EmbedInteropTypes>True</EmbedInteropTypes>
     </Reference>
     <Reference Include="Microsoft.VisualStudio.VCProject, Version=11.0.0.0, Culture=neutral, PublicKeyToken=b03f5f7f11d50a3a" Condition="'$(Configuration)' == 'VS11' Or '$(Configuration)' == 'VS11d'">
       <EmbedInteropTypes>True</EmbedInteropTypes>
     </Reference>
+    <Reference Include="Microsoft.VisualStudio.VCProjectEngine, Version=14.0.0.0, Culture=neutral, PublicKeyToken=b03f5f7f11d50a3a" Condition="'$(Configuration)' == 'VS14' Or '$(Configuration)' == 'VS14d'">
+      <EmbedInteropTypes>True</EmbedInteropTypes>
+    </Reference>
     <Reference Include="Microsoft.VisualStudio.VCProjectEngine, Version=12.0.0.0, Culture=neutral, PublicKeyToken=b03f5f7f11d50a3a" Condition="'$(Configuration)' == 'VS12' Or '$(Configuration)' == 'VS12d'">
       <EmbedInteropTypes>True</EmbedInteropTypes>
     </Reference>
@@ -172,7 +179,7 @@
     <Compile Include="GlobalSuppressions.cs" />
     <Compile Include="StartupProjectSelector.cs" />
     <Compile Include="UnrealVSOptions.cs">
-      <SubType>Component</SubType>
+      <SubType>component</SubType>
     </Compile>
     <Compile Include="UnrealVSPackage.cs" />
     <Compile Include="Properties\AssemblyInfo.cs" />
@@ -209,6 +216,8 @@
     <None Include="Key.snk" />
   </ItemGroup>
   <ItemGroup>
+    <Content Include="TargetManifests\VS14.source.extension.vsixmanifest" />
+    <Content Include="TargetManifests\VS14d.source.extension.vsixmanifest" />
     <None Include="Generated\source.extension.vsixmanifest" />
     <None Include="Resources\BatchBuilder.png" />
     <Resource Include="Resources\BatchBuilderBackground.png" />
@@ -245,12 +254,16 @@
       <Generator>MSBuild:Compile</Generator>
       <Generator>MSBuild:Compile</Generator>
       <SubType>Designer</SubType>
+      <Generator>MSBuild:Compile</Generator>
+      <SubType>Designer</SubType>
     </Page>
     <Page Include="BuildStatusIcon.xaml">
       <SubType>Designer</SubType>
       <Generator>MSBuild:Compile</Generator>
       <Generator>MSBuild:Compile</Generator>
       <SubType>Designer</SubType>
+      <Generator>MSBuild:Compile</Generator>
+      <SubType>Designer</SubType>
     </Page>
   </ItemGroup>
   <PropertyGroup>
@@ -283,6 +296,25 @@
     <DeployExtension>True</DeployExtension>
     <CreateVsixContainer>True</CreateVsixContainer>
   </PropertyGroup>
+  <PropertyGroup Condition=" '$(Configuration)|$(Platform)' == 'VS14d|AnyCPU' ">
+    <DebugSymbols>true</DebugSymbols>
+    <OutputPath>bin\VS14d\</OutputPath>
+    <DefineConstants>TRACE;DEBUG;VS14</DefineConstants>
+    <DebugType>full</DebugType>
+    <PlatformTarget>AnyCPU</PlatformTarget>
+    <ErrorReport>prompt</ErrorReport>
+    <CodeAnalysisRuleSet>MinimumRecommendedRules.ruleset</CodeAnalysisRuleSet>
+  </PropertyGroup>
+  <PropertyGroup Condition=" '$(Configuration)|$(Platform)' == 'VS14|AnyCPU' ">
+    <OutputPath>..\..\..\Extras\UnrealVS\VS2015\</OutputPath>
+    <DefineConstants>TRACE;CODE_ANALYSIS;VS14</DefineConstants>
+    <Optimize>true</Optimize>
+    <DebugType>pdbonly</DebugType>
+    <PlatformTarget>AnyCPU</PlatformTarget>
+    <RunCodeAnalysis>true</RunCodeAnalysis>
+    <ErrorReport>prompt</ErrorReport>
+    <CodeAnalysisRuleSet>MinimumRecommendedRules.ruleset</CodeAnalysisRuleSet>
+  </PropertyGroup>
   <Import Project="$(MSBuildBinPath)\Microsoft.CSharp.targets" />
   <Import Project="$(VSToolsPath)\VSSDK\Microsoft.VsSDK.targets" Condition="'$(VSToolsPath)' != ''" />
   <Import Project="$(MSBuildExtensionsPath32)\Microsoft\VisualStudio\v11.0\VSSDK\Microsoft.VsSDK.targets" Condition="false" />
diff --git a/Engine/Source/Programs/UnrealVS/UnrealVSPackage.cs b/Engine/Source/Programs/UnrealVS/UnrealVSPackage.cs
--- a/Engine/Source/Programs/UnrealVS/UnrealVSPackage.cs
+++ b/Engine/Source/Programs/UnrealVS/UnrealVSPackage.cs
@@ -55,7 +55,7 @@ namespace UnrealVS
 	{
 		/** Constants */
 
-		private const string VersionString = "v1.38";
+		private const string VersionString = "v1.39";
 		private const string UnrealSolutionFileNamePrefix = "UE4";
 		private const string ExtensionName = "UnrealVS";
 		private const string CommandLineOptionKey = ExtensionName + "CommandLineMRU";
