UnrealVS 1.41 - added logging to track possible hang when starting up VS or opening a solution.

diff --git a/Engine/Source/Programs/UnrealVS/Properties/AssemblyInfo.cs b/Engine/Source/Programs/UnrealVS/Properties/AssemblyInfo.cs
--- a/Engine/Source/Programs/UnrealVS/Properties/AssemblyInfo.cs
+++ b/Engine/Source/Programs/UnrealVS/Properties/AssemblyInfo.cs
@@ -31,8 +31,8 @@ using System.Runtime.InteropServices;
 // You can specify all the values or you can default the Revision and Build Numbers 
 // by using the '*' as shown below:
 
-[assembly: AssemblyVersion("1.40.0.0")]
-[assembly: AssemblyFileVersion("1.40.0.0")]
+[assembly: AssemblyVersion("1.41.0.0")]
+[assembly: AssemblyFileVersion("1.41.0.0")]
 
 
 
diff --git a/Engine/Source/Programs/UnrealVS/UnrealVS.csproj b/Engine/Source/Programs/UnrealVS/UnrealVS.csproj
--- a/Engine/Source/Programs/UnrealVS/UnrealVS.csproj
+++ b/Engine/Source/Programs/UnrealVS/UnrealVS.csproj
@@ -179,7 +179,7 @@
     <Compile Include="GlobalSuppressions.cs" />
     <Compile Include="StartupProjectSelector.cs" />
     <Compile Include="UnrealVSOptions.cs">
-      <SubType>component</SubType>
+      <SubType>Component</SubType>
     </Compile>
     <Compile Include="UnrealVSPackage.cs" />
     <Compile Include="Properties\AssemblyInfo.cs" />
diff --git a/Engine/Source/Programs/UnrealVS/UnrealVSPackage.cs b/Engine/Source/Programs/UnrealVS/UnrealVSPackage.cs
--- a/Engine/Source/Programs/UnrealVS/UnrealVSPackage.cs
+++ b/Engine/Source/Programs/UnrealVS/UnrealVSPackage.cs
@@ -55,7 +55,7 @@ namespace UnrealVS
 	{
 		/** Constants */
 
-		private const string VersionString = "v1.40";
+		private const string VersionString = "v1.41";
 		private const string UnrealSolutionFileNamePrefix = "UE4";
 		private const string ExtensionName = "UnrealVS";
 		private const string CommandLineOptionKey = ExtensionName + "CommandLineMRU";
diff --git a/Engine/Source/Programs/UnrealVS/Utils.cs b/Engine/Source/Programs/UnrealVS/Utils.cs
--- a/Engine/Source/Programs/UnrealVS/Utils.cs
+++ b/Engine/Source/Programs/UnrealVS/Utils.cs
@@ -603,12 +603,15 @@ namespace UnrealVS
 			if (Folder != CachedUProjectRootFolder)
 			{
 				Logging.WriteLine("GetUProjects: recaching uproject paths...");
+                DateTime Start = DateTime.Now;
 
 				CachedUProjectRootFolder = Folder;
 				CachedUProjectPaths = Directory.GetFiles(Folder, "*." + UProjectExtension, SearchOption.AllDirectories);
 				CachedUProjects = null;
 
-				Logging.WriteLine("    DONE");
+                TimeSpan TimeTaken = DateTime.Now - Start;
+                Logging.WriteLine(string.Format("Directory.GetFiles took {0} sec", TimeTaken.TotalSeconds));
+                Logging.WriteLine("    DONE");
 			}
 
 			if (CachedUProjects == null)
