Merging fixes to UnrealVS with VS2017 from 4.15 branch.

diff --git a/Engine/Source/Programs/UnrealVS/CommandLineEditor.cs b/Engine/Source/Programs/UnrealVS/CommandLineEditor.cs
--- a/Engine/Source/Programs/UnrealVS/CommandLineEditor.cs
+++ b/Engine/Source/Programs/UnrealVS/CommandLineEditor.cs
@@ -144,43 +144,41 @@ namespace UnrealVS
 			string Text = "";
 
 			IVsHierarchy ProjectHierarchy;
-			UnrealVSPackage.Instance.SolutionBuildManager.get_StartupProject( out ProjectHierarchy );
-			if( ProjectHierarchy != null )
+			if (UnrealVSPackage.Instance.SolutionBuildManager.get_StartupProject(out ProjectHierarchy) == VSConstants.S_OK && ProjectHierarchy != null)
 			{
-				var SelectedStartupProject = Utils.HierarchyObjectToProject( ProjectHierarchy );
-
+				Project SelectedStartupProject = Utils.HierarchyObjectToProject(ProjectHierarchy);
 				if (SelectedStartupProject != null)
 				{
-					var CommandLineArgumentsProperty = GetProjectCommandLineProperty(SelectedStartupProject);
-
-					if (CommandLineArgumentsProperty != null)
+					Configuration SelectedConfiguration = SelectedStartupProject.ConfigurationManager.ActiveConfiguration;
+					if (SelectedConfiguration != null)
 					{
-						var CommandLineArguments = (string) CommandLineArgumentsProperty.Value;
-
-						// for "Game" projects automatically remove the game project filename from the start of the command line
-						var ActiveConfiguration = (SolutionConfiguration2) UnrealVSPackage.Instance.DTE.Solution.SolutionBuild.ActiveConfiguration;
-						if (UnrealVSPackage.Instance.IsUE4Loaded && Utils.IsGameProject(SelectedStartupProject) && Utils.HasUProjectCommandLineArg(ActiveConfiguration.Name))
+						IVsBuildPropertyStorage PropertyStorage = ProjectHierarchy as IVsBuildPropertyStorage;
+						if (PropertyStorage != null)
 						{
-							var AutoPrefix = Utils.GetAutoUProjectCommandLinePrefix(SelectedStartupProject);
+							// Query the property store for the debugger arguments
+							string ConfigurationName = String.Format("{0}|{1}", SelectedConfiguration.ConfigurationName, SelectedConfiguration.PlatformName);
+							if (PropertyStorage.GetPropertyValue("LocalDebuggerCommandArguments", ConfigurationName, (uint)_PersistStorageType.PST_USER_FILE, out Text) != VSConstants.S_OK)
+							{
+								if (PropertyStorage.GetPropertyValue("StartArguments", ConfigurationName, (uint)_PersistStorageType.PST_USER_FILE, out Text) != VSConstants.S_OK)
+								{
+									Text = InvalidProjectString;
+								}
+							}
 
-							if (!string.IsNullOrEmpty(AutoPrefix))
+							// for "Game" projects automatically remove the game project filename from the start of the command line
+							var ActiveConfiguration = (SolutionConfiguration2)UnrealVSPackage.Instance.DTE.Solution.SolutionBuild.ActiveConfiguration;
+							if (UnrealVSPackage.Instance.IsUE4Loaded && Utils.IsGameProject(SelectedStartupProject) && Utils.HasUProjectCommandLineArg(ActiveConfiguration.Name))
 							{
-								if (CommandLineArguments.Trim().StartsWith(AutoPrefix, StringComparison.OrdinalIgnoreCase))
+								string AutoPrefix = Utils.GetAutoUProjectCommandLinePrefix(SelectedStartupProject);
+								if (!string.IsNullOrEmpty(AutoPrefix))
 								{
-									CommandLineArguments = CommandLineArguments.Trim().Substring(AutoPrefix.Length).Trim();
+									if (Text.Trim().StartsWith(AutoPrefix, StringComparison.OrdinalIgnoreCase))
+									{
+										Text = Text.Trim().Substring(AutoPrefix.Length).Trim();
+									}
 								}
-								//else if (CommandLineArguments.Trim().StartsWith(SelectedStartupProject.Name + " ", StringComparison.OrdinalIgnoreCase))
-								//{
-								//	CommandLineArguments = CommandLineArguments.Trim().Substring(SelectedStartupProject.Name.Length + 1).Trim();
-								//}
 							}
 						}
-
-						Text = CommandLineArguments;
-					}
-					else
-					{
-						Text = InvalidProjectString;
 					}
 				}
 			}
@@ -250,41 +248,52 @@ namespace UnrealVS
 		private void CommitCommandLineText(string CommandLine)
 		{
 			IVsHierarchy ProjectHierarchy;
-			UnrealVSPackage.Instance.SolutionBuildManager.get_StartupProject(out ProjectHierarchy);
-			if (ProjectHierarchy != null)
+			if (UnrealVSPackage.Instance.SolutionBuildManager.get_StartupProject(out ProjectHierarchy) == VSConstants.S_OK && ProjectHierarchy != null)
 			{
-				var SelectedStartupProject = Utils.HierarchyObjectToProject(ProjectHierarchy);
-
+				Project SelectedStartupProject = Utils.HierarchyObjectToProject(ProjectHierarchy);
 				if (SelectedStartupProject != null)
 				{
-					var CommandLineArgumentsProperty = GetProjectCommandLineProperty(SelectedStartupProject);
-
-					if (CommandLineArgumentsProperty != null)
+					Configuration SelectedConfiguration = SelectedStartupProject.ConfigurationManager.ActiveConfiguration;
+					if (SelectedConfiguration != null)
 					{
-						string FullCommandLine = CommandLine;
-
-						// for "Game" projects automatically remove the game project filename from the start of the command line
-						var ActiveConfiguration = (SolutionConfiguration2)UnrealVSPackage.Instance.DTE.Solution.SolutionBuild.ActiveConfiguration;
-						if (UnrealVSPackage.Instance.IsUE4Loaded && Utils.IsGameProject(SelectedStartupProject) && Utils.HasUProjectCommandLineArg(ActiveConfiguration.Name))
+						IVsBuildPropertyStorage PropertyStorage = ProjectHierarchy as IVsBuildPropertyStorage;
+						if (PropertyStorage != null)
 						{
-							var AutoPrefix = Utils.GetAutoUProjectCommandLinePrefix(SelectedStartupProject);
-							
-							if (FullCommandLine.IndexOf(Utils.UProjectExtension, StringComparison.OrdinalIgnoreCase) < 0 &&
-								string.Compare(FullCommandLine, SelectedStartupProject.Name, StringComparison.OrdinalIgnoreCase) != 0 &&
-								!FullCommandLine.StartsWith(SelectedStartupProject.Name + " ", StringComparison.OrdinalIgnoreCase))
+							string FullCommandLine = CommandLine;
+
+							// for "Game" projects automatically remove the game project filename from the start of the command line
+							var ActiveConfiguration = (SolutionConfiguration2)UnrealVSPackage.Instance.DTE.Solution.SolutionBuild.ActiveConfiguration;
+							if (UnrealVSPackage.Instance.IsUE4Loaded && Utils.IsGameProject(SelectedStartupProject) && Utils.HasUProjectCommandLineArg(ActiveConfiguration.Name))
 							{
-								// Committed command line does not specify a .uproject
-								FullCommandLine = AutoPrefix + " " + FullCommandLine;
+								string AutoPrefix = Utils.GetAutoUProjectCommandLinePrefix(SelectedStartupProject);
+								if (FullCommandLine.IndexOf(Utils.UProjectExtension, StringComparison.OrdinalIgnoreCase) < 0 &&
+									string.Compare(FullCommandLine, SelectedStartupProject.Name, StringComparison.OrdinalIgnoreCase) != 0 &&
+									!FullCommandLine.StartsWith(SelectedStartupProject.Name + " ", StringComparison.OrdinalIgnoreCase))
+								{
+									// Committed command line does not specify a .uproject
+									FullCommandLine = AutoPrefix + " " + FullCommandLine;
+								}
 							}
-						}
 
-						Utils.SetPropertyValue(CommandLineArgumentsProperty, FullCommandLine);
+							// Get the project kind. C++ projects store the debugger arguments differently to other project types.
+							string ProjectKind = SelectedStartupProject.Kind;
+
+							// Update the property
+							string ProjectConfigurationName = String.Format("{0}|{1}", SelectedConfiguration.ConfigurationName, SelectedConfiguration.PlatformName);
+							if (String.Equals(ProjectKind, "{8BC9CEB8-8B4A-11D0-8D11-00A0C91BC942}", StringComparison.InvariantCultureIgnoreCase))
+							{
+								PropertyStorage.SetPropertyValue("LocalDebuggerCommandArguments", ProjectConfigurationName, (uint)_PersistStorageType.PST_USER_FILE, FullCommandLine);
+							}
+							else
+							{
+								PropertyStorage.SetPropertyValue("StartArguments", ProjectConfigurationName, (uint)_PersistStorageType.PST_USER_FILE, FullCommandLine);
+							}
+						}
 					}
 				}
 			}
 
 			CommitCommandLineToMRU(CommandLine);
-
 		}
 
 		private void CommitCommandLineToMRU(string CommandLine)
@@ -312,23 +321,6 @@ namespace UnrealVS
 			Marshal.GetNativeVariantForObject(ComboList.ToArray(), OleArgs.OutValue);
 		}
 
-		/// <summary>
-		/// Helper function to get the correct property from a project that represents its command line arguments
-		/// </summary>
-		private static Property GetProjectCommandLineProperty(Project InProject)
-		{
-			// C++ projects use "CommandArguments" as the property name
-			var CommandLineArgumentsProperty = Utils.GetProjectConfigProperty(InProject, null, "CommandArguments");
-
-			// C# projects use "StartArguments" as the property name
-			if (CommandLineArgumentsProperty == null)
-			{
-				CommandLineArgumentsProperty = Utils.GetProjectConfigProperty(InProject, null, "StartArguments");
-			}
-
-			return CommandLineArgumentsProperty;
-		}
-
 		/// List of MRU strings to show in the combobox drop-down list
 		private readonly List<string> ComboList = new List<string>();
 
diff --git a/Engine/Source/Programs/UnrealVS/Properties/AssemblyInfo.cs b/Engine/Source/Programs/UnrealVS/Properties/AssemblyInfo.cs
--- a/Engine/Source/Programs/UnrealVS/Properties/AssemblyInfo.cs
+++ b/Engine/Source/Programs/UnrealVS/Properties/AssemblyInfo.cs
@@ -31,8 +31,8 @@ using System.Runtime.InteropServices;
 // You can specify all the values or you can default the Revision and Build Numbers 
 // by using the '*' as shown below:
 
-[assembly: AssemblyVersion("1.45.0.0")]
-[assembly: AssemblyFileVersion("1.45.0.0")]
+[assembly: AssemblyVersion("1.46.0.0")]
+[assembly: AssemblyFileVersion("1.46.0.0")]
 
 
 
diff --git a/Engine/Source/Programs/UnrealVS/UnrealVSPackage.cs b/Engine/Source/Programs/UnrealVS/UnrealVSPackage.cs
--- a/Engine/Source/Programs/UnrealVS/UnrealVSPackage.cs
+++ b/Engine/Source/Programs/UnrealVS/UnrealVSPackage.cs
@@ -30,6 +30,9 @@ namespace UnrealVS
 	// match the ItemGroup inside the UnrealVS.csproj MSBuild file (hand-typed.)
 	[ProvideMenuResource( "Menus.ctmenu", 1 )]
 
+	// Force the package to load whenever a solution exists
+	[ProvideAutoLoad(UIContextGuids80.SolutionExists)]
+
 	// Register the settings implementing class as providing support to UnrealVSPackage.
 	//[ProvideProfile(typeof (ProfileManager), "UnrealVS", "UnrealVSPackage", 110, 113, false)]
 
@@ -56,7 +59,7 @@ namespace UnrealVS
 	{
 		/** Constants */
 
-		private const string VersionString = "v1.45";
+		private const string VersionString = "v1.46";
 		private const string UnrealSolutionFileNamePrefix = "UE4";
 		private const string ExtensionName = "UnrealVS";
 		private const string CommandLineOptionKey = ExtensionName + "CommandLineMRU";
diff --git a/Engine/Source/Programs/UnrealVS/UnrealVS_2017.csproj b/Engine/Source/Programs/UnrealVS/UnrealVS_2017.csproj
--- a/Engine/Source/Programs/UnrealVS/UnrealVS_2017.csproj
+++ b/Engine/Source/Programs/UnrealVS/UnrealVS_2017.csproj
@@ -185,6 +185,7 @@
     </EmbeddedResource>
   </ItemGroup>
   <ItemGroup>
+    <None Include="packages.config" />
     <None Include="TargetManifests\VS11.source.extension.vsixmanifest">
       <SubType>Designer</SubType>
     </None>
@@ -321,6 +322,9 @@
     <PlatformTarget>AnyCPU</PlatformTarget>
     <ErrorReport>prompt</ErrorReport>
     <CodeAnalysisRuleSet>MinimumRecommendedRules.ruleset</CodeAnalysisRuleSet>
+    <DeployExtension>False</DeployExtension>
+    <CopyVsixExtensionFiles>True</CopyVsixExtensionFiles>
+    <CopyVsixExtensionLocation>[REDACTED_USER_HOME]\AppData\Local\Microsoft\VisualStudio\15.0_12f23ee0Exp\Extensions\Epic Games, Inc\UnrealVS\1.45</CopyVsixExtensionLocation>
   </PropertyGroup>
   <Import Project="$(MSBuildBinPath)\Microsoft.CSharp.targets" />
   <Import Project="$(VSToolsPath)\VSSDK\Microsoft.VsSDK.targets" Condition="'$(VSToolsPath)' != ''" />
diff --git a/Engine/Source/Programs/UnrealVS/Utils.cs b/Engine/Source/Programs/UnrealVS/Utils.cs
--- a/Engine/Source/Programs/UnrealVS/Utils.cs
+++ b/Engine/Source/Programs/UnrealVS/Utils.cs
@@ -10,6 +10,7 @@ using System.Linq.Expressions;
 using VSLangProj;
 using System.IO;
 using EnvDTE80;
+using Microsoft.VisualStudio.Shell;
 
 namespace UnrealVS
 {
@@ -103,36 +104,6 @@ namespace UnrealVS
 			return CfgProvider2;
 		}
 
-		/// <summary>
-		/// Locates a specific project config property for the active configuration and returns it (or null if not found.)
-		/// </summary>
-		/// <param name="Project">Project to search the active configuration for the property</param>
-		/// <param name="Configuration">Project configuration to edit, or null to use the "active" configuration</param>
-		/// <param name="PropertyName">Name of the property</param>
-		/// <returns>Property object or null if not found</returns>
-		public static Property GetProjectConfigProperty(Project Project, Configuration Configuration, string PropertyName)
-		{
-			if (Configuration == null)
-			{
-				Configuration = Project.ConfigurationManager.ActiveConfiguration;
-			}
-			if (Configuration != null && Configuration.Properties != null)
-			{
-				var Properties = Configuration.Properties;
-				foreach (var RawProperty in Properties)
-				{
-					var Property = (Property)RawProperty;
-					if (Property.Name.Equals(PropertyName, StringComparison.InvariantCultureIgnoreCase))
-					{
-						return Property;
-					}
-				}
-			}
-
-			// Not found
-			return null;
-		}
-
 		/// <summary>
 		/// Locates a specific project property for the active configuration and returns it (or null if not found.)
 		/// </summary>
@@ -527,7 +498,10 @@ namespace UnrealVS
 		public static bool SelectProjectInSolutionExplorer(Project Project)
 		{
 			UnrealVSPackage.Instance.DTE.ExecuteCommand("View.SolutionExplorer");
-			Project.ParentProjectItem.ExpandView();
+			if (Project.ParentProjectItem != null)
+			{
+				Project.ParentProjectItem.ExpandView();
+			}
 
 			UIHierarchy SolutionExplorerHierarachy = UnrealVSPackage.Instance.DTE2.ToolWindows.SolutionExplorer;
 			Utils.UITreeItem SolutionExplorerTree = Utils.GetUIHierarchyTree(SolutionExplorerHierarachy);
