UnrealVS: Remove absolute path to output directory.

#rnx

diff --git a/Engine/Source/Programs/UnrealVS/UnrealVS_2017.csproj b/Engine/Source/Programs/UnrealVS/UnrealVS_2017.csproj
--- a/Engine/Source/Programs/UnrealVS/UnrealVS_2017.csproj
+++ b/Engine/Source/Programs/UnrealVS/UnrealVS_2017.csproj
@@ -314,8 +314,7 @@
     <ErrorReport>prompt</ErrorReport>
     <CodeAnalysisRuleSet>MinimumRecommendedRules.ruleset</CodeAnalysisRuleSet>
     <DeployExtension>True</DeployExtension>
-    <CopyVsixExtensionFiles>True</CopyVsixExtensionFiles>
-    <CopyVsixExtensionLocation>[REDACTED_USER_HOME]\AppData\Local\Microsoft\VisualStudio\15.0_12f23ee0Exp\Extensions\Epic Games, Inc\UnrealVS\1.48</CopyVsixExtensionLocation>
+    <CopyVsixExtensionFiles>False</CopyVsixExtensionFiles>
   </PropertyGroup>
   <Import Project="$(MSBuildBinPath)\Microsoft.CSharp.targets" />
   <Import Project="$(VSToolsPath)\VSSDK\Microsoft.VsSDK.targets" Condition="'$(VSToolsPath)' != ''" />
