Fix performance problem with p4vs Helix visual studio plugin, it was due to the UnrealVs plugin listening to all property change events and updating command line combo everytime.

diff --git a/Engine/Source/Programs/UnrealVS/CommandLineEditor.cs b/Engine/Source/Programs/UnrealVS/CommandLineEditor.cs
--- a/Engine/Source/Programs/UnrealVS/CommandLineEditor.cs
+++ b/Engine/Source/Programs/UnrealVS/CommandLineEditor.cs
@@ -13,7 +13,7 @@ using EnvDTE;
 using System.Runtime.InteropServices;
 using System.Collections.Generic;
 using EnvDTE80;
-
+using System.Timers;
 
 namespace UnrealVS
 {
@@ -57,16 +57,40 @@ namespace UnrealVS
 		/// </summary>
 		public void OnStartupProjectPropertyChanged(UInt32 itemid, Int32 propid, UInt32 flags)
 		{
+			// Filter out Helix VS plugin sending thousands of VSHPROPID_StateIconIndex in UE4.sln
+			// This event type is not usefull at all for the commandline editor
+			if (propid == (Int32)__VSHPROPID.VSHPROPID_StateIconIndex)
+			{
+				return;
+			}
+
 			IVsHierarchy ProjectHierarchy;
 			UnrealVSPackage.Instance.SolutionBuildManager.get_StartupProject(out ProjectHierarchy);
 			if (ProjectHierarchy != null)
 			{
 				
 				// @todo: filter this so that we only respond to changes in the command line property
-				UpdateCommandLineCombo();
+				// Setup a timer to prevent any more performance problem from spamming events
+				if( UpdateCommandLineComboTimer == null )
+				{
+					UpdateCommandLineComboTimer = new System.Timers.Timer(1000);
+					UpdateCommandLineComboTimer.AutoReset = false;
+					UpdateCommandLineComboTimer.Elapsed += OnUpdateCommandLineCombo;
+				}
+				// Restart timer to raise the event only after 1s of no notification
+				UpdateCommandLineComboTimer.Stop();
+				UpdateCommandLineComboTimer.Start();
 			}
 		}
 
+		/// <summary>
+		/// Timer callback to UpdateCommandLineCombo
+		/// </summary>
+		private void OnUpdateCommandLineCombo(Object source, ElapsedEventArgs e)
+		{
+			UpdateCommandLineCombo();
+		}
+
 		/// <summary>
 		/// Reads options out of the solution file.
 		/// </summary>
@@ -366,5 +390,8 @@ namespace UnrealVS
 
 		/// Used to store the user edited commandline mid-edit, in the combo handler
 		private string DesiredCommandLine;
+
+		/// used to shield against too many property changes in the future
+		private System.Timers.Timer UpdateCommandLineComboTimer = null;
 	}
 }
diff --git a/Engine/Source/Programs/UnrealVS/Properties/AssemblyInfo.cs b/Engine/Source/Programs/UnrealVS/Properties/AssemblyInfo.cs
--- a/Engine/Source/Programs/UnrealVS/Properties/AssemblyInfo.cs
+++ b/Engine/Source/Programs/UnrealVS/Properties/AssemblyInfo.cs
@@ -31,8 +31,8 @@ using System.Runtime.InteropServices;
 // You can specify all the values or you can default the Revision and Build Numbers 
 // by using the '*' as shown below:
 
-[assembly: AssemblyVersion("1.53.0.0")]
-[assembly: AssemblyFileVersion("1.53.0.0")]
+[assembly: AssemblyVersion("1.54.0.0")]
+[assembly: AssemblyFileVersion("1.54.0.0")]
 
 
 
