Fix crash with newly added code.

diff --git a/Engine/Source/Programs/UnrealVS/CommandLineEditor.cs b/Engine/Source/Programs/UnrealVS/CommandLineEditor.cs
--- a/Engine/Source/Programs/UnrealVS/CommandLineEditor.cs
+++ b/Engine/Source/Programs/UnrealVS/CommandLineEditor.cs
@@ -88,7 +88,7 @@ namespace UnrealVS
 		/// </summary>
 		private void OnUpdateCommandLineCombo(Object source, ElapsedEventArgs e)
 		{
-			UpdateCommandLineCombo();
+            ThreadHelper.Generic.BeginInvoke(UpdateCommandLineCombo);
 		}
 
 		/// <summary>
