Add build of UnrealVS for Visual Studio 2019.

diff --git a/Engine/Source/Programs/UnrealVS/UnrealVS_2019.csproj b/Engine/Source/Programs/UnrealVS/UnrealVS_2019.csproj
new file mode 100644
--- /dev/null
+++ b/Engine/Source/Programs/UnrealVS/UnrealVS_2019.csproj
@@ -0,0 +1,357 @@
+﻿<?xml version="1.0" encoding="utf-8"?>
+<Project DefaultTargets="Build" xmlns="http://schemas.microsoft.com/developer/msbuild/2003" ToolsVersion="14.0">
+  <Import Project="$(MSBuildExtensionsPath)\$(MSBuildToolsVersion)\Microsoft.Common.props" Condition="Exists('$(MSBuildExtensionsPath)\$(MSBuildToolsVersion)\Microsoft.Common.props')" />
+  <PropertyGroup>
+    <Configuration Condition=" '$(Configuration)' == '' ">Debug</Configuration>
+    <Platform Condition=" '$(Platform)' == '' ">AnyCPU</Platform>
+    <SchemaVersion>2.0</SchemaVersion>
+    <ProjectGuid>{C0EE3E24-0255-4D0E-9369-82247AD5E2DD}</ProjectGuid>
+    <ProjectTypeGuids>{82b43b9b-a64c-4715-b499-d71e9ca2bd60};{60dc8134-eba5-43b8-bcc9-bb4bc16c2548};{FAE04EC0-301F-11D3-BF4B-00C04F79EFBC}</ProjectTypeGuids>
+    <OutputType>Library</OutputType>
+    <AppDesignerFolder>Properties</AppDesignerFolder>
+    <RootNamespace>UnrealVS</RootNamespace>
+    <AssemblyName>UnrealVS</AssemblyName>
+    <SignAssembly>True</SignAssembly>
+    <AssemblyOriginatorKeyFile>Key.snk</AssemblyOriginatorKeyFile>
+    <TargetFrameworkVersion>v4.7.2</TargetFrameworkVersion>
+    <SccProjectName>
+    </SccProjectName>
+    <SccLocalPath>
+    </SccLocalPath>
+    <SccAuxPath>
+    </SccAuxPath>
+    <SccProvider>
+    </SccProvider>
+    <MinimumVisualStudioVersion>15.0</MinimumVisualStudioVersion>
+    <OldToolsVersion>14.0</OldToolsVersion>
+    <TargetFrameworkProfile />
+    <FileUpgradeFlags>
+    </FileUpgradeFlags>
+    <UpgradeBackupLocation>
+    </UpgradeBackupLocation>
+    <IsWebBootstrapper>false</IsWebBootstrapper>
+    <PublishUrl>publish\</PublishUrl>
+    <Install>true</Install>
+    <InstallFrom>Disk</InstallFrom>
+    <UpdateEnabled>false</UpdateEnabled>
+    <UpdateMode>Foreground</UpdateMode>
+    <UpdateInterval>7</UpdateInterval>
+    <UpdateIntervalUnits>Days</UpdateIntervalUnits>
+    <UpdatePeriodically>false</UpdatePeriodically>
+    <UpdateRequired>false</UpdateRequired>
+    <MapFileExtensions>true</MapFileExtensions>
+    <ApplicationRevision>0</ApplicationRevision>
+    <ApplicationVersion>1.0.0.%2a</ApplicationVersion>
+    <UseApplicationTrust>false</UseApplicationTrust>
+    <BootstrapperEnabled>true</BootstrapperEnabled>
+  </PropertyGroup>
+  <PropertyGroup Condition=" '$(Configuration)|$(Platform)' == 'VS12d|AnyCPU' ">
+    <DebugSymbols>true</DebugSymbols>
+    <DebugType>full</DebugType>
+    <Optimize>false</Optimize>
+    <OutputPath>bin\VS12d\</OutputPath>
+    <DefineConstants>TRACE;DEBUG;VS12</DefineConstants>
+    <ErrorReport>prompt</ErrorReport>
+    <WarningLevel>4</WarningLevel>
+    <Prefer32Bit>false</Prefer32Bit>
+    <CopyVsixExtensionFiles>False</CopyVsixExtensionFiles>
+    <DeployExtension>True</DeployExtension>
+    <CreateVsixContainer>True</CreateVsixContainer>
+  </PropertyGroup>
+  <PropertyGroup Condition=" '$(Configuration)|$(Platform)' == 'VS12|AnyCPU' ">
+    <DebugType>pdbonly</DebugType>
+    <Optimize>true</Optimize>
+    <OutputPath>..\..\..\Extras\UnrealVS\VS2013\</OutputPath>
+    <DefineConstants>TRACE;CODE_ANALYSIS;VS12</DefineConstants>
+    <ErrorReport>prompt</ErrorReport>
+    <WarningLevel>4</WarningLevel>
+    <RunCodeAnalysis>true</RunCodeAnalysis>
+    <Prefer32Bit>false</Prefer32Bit>
+  </PropertyGroup>
+  <ItemGroup>
+    <VSCTCompile Include="UnrealVS.vsct">
+      <ResourceName>Menus.ctmenu</ResourceName>
+      <SubType>Designer</SubType>
+    </VSCTCompile>
+  </ItemGroup>
+  <ItemGroup>
+    <Reference Include="envdte, Version=8.0.0.0, Culture=neutral, PublicKeyToken=b03f5f7f11d50a3a">
+      <EmbedInteropTypes>False</EmbedInteropTypes>
+    </Reference>
+    <Reference Include="envdte80, Version=8.0.0.0, Culture=neutral, PublicKeyToken=b03f5f7f11d50a3a">
+      <EmbedInteropTypes>True</EmbedInteropTypes>
+    </Reference>
+    <Reference Include="Microsoft.CSharp" />
+    <Reference Include="Microsoft.VisualStudio.OLE.Interop, Version=7.1.40304.0, Culture=neutral, PublicKeyToken=b03f5f7f11d50a3a" />
+    <Reference Include="Microsoft.VisualStudio.Shell.15.0, Version=16.0.0.0, Culture=neutral, PublicKeyToken=b03f5f7f11d50a3a, processorArchitecture=MSIL" />
+    <Reference Include="Microsoft.VisualStudio.Shell.Framework, Version=16.0.0.0, Culture=neutral, PublicKeyToken=b03f5f7f11d50a3a, processorArchitecture=MSIL" />
+    <Reference Include="Microsoft.VisualStudio.Shell.Interop, Version=7.1.40304.0, Culture=neutral, PublicKeyToken=b03f5f7f11d50a3a" />
+    <Reference Include="Microsoft.VisualStudio.Shell.Interop.10.0, Version=10.0.0.0, Culture=neutral, PublicKeyToken=b03f5f7f11d50a3a, processorArchitecture=MSIL">
+      <EmbedInteropTypes>True</EmbedInteropTypes>
+    </Reference>
+    <Reference Include="Microsoft.VisualStudio.Shell.Interop.12.0, Version=12.0.0.0, Culture=neutral, PublicKeyToken=b03f5f7f11d50a3a, processorArchitecture=MSIL">
+      <EmbedInteropTypes>True</EmbedInteropTypes>
+    </Reference>
+    <Reference Include="Microsoft.VisualStudio.Shell.Interop.15.3.DesignTime, Version=15.3.0.0, Culture=neutral, PublicKeyToken=b03f5f7f11d50a3a, processorArchitecture=MSIL">
+      <EmbedInteropTypes>True</EmbedInteropTypes>
+    </Reference>
+    <Reference Include="Microsoft.VisualStudio.Shell.Interop.8.0, Version=8.0.0.0, Culture=neutral, PublicKeyToken=b03f5f7f11d50a3a" />
+    <Reference Include="Microsoft.VisualStudio.Shell.Interop.9.0, Version=9.0.0.0, Culture=neutral, PublicKeyToken=b03f5f7f11d50a3a" />
+    <Reference Include="Microsoft.VisualStudio.TextManager.Interop, Version=7.1.40304.0, Culture=neutral, PublicKeyToken=b03f5f7f11d50a3a" />
+    <Reference Include="Microsoft.VisualStudio.VCProject, Version=16.0.0.0, Culture=neutral, PublicKeyToken=b03f5f7f11d50a3a">
+      <EmbedInteropTypes>True</EmbedInteropTypes>
+    </Reference>
+    <Reference Include="Microsoft.VisualStudio.VCProjectEngine, Version=16.0.0.0, Culture=neutral, PublicKeyToken=b03f5f7f11d50a3a">
+      <EmbedInteropTypes>True</EmbedInteropTypes>
+    </Reference>
+    <Reference Include="PresentationCore" />
+    <Reference Include="PresentationFramework" />
+    <Reference Include="System" />
+    <Reference Include="System.Core" />
+    <Reference Include="System.Data" />
+    <Reference Include="System.Design" />
+    <Reference Include="System.Drawing" />
+    <Reference Include="System.Windows.Forms" />
+    <Reference Include="System.Xaml" />
+    <Reference Include="System.Xml" />
+    <Reference Include="VSLangProj, Version=7.0.3300.0, Culture=neutral, PublicKeyToken=b03f5f7f11d50a3a">
+      <EmbedInteropTypes>True</EmbedInteropTypes>
+    </Reference>
+    <Reference Include="VSLangProj80, Version=8.0.0.0, Culture=neutral, PublicKeyToken=b03f5f7f11d50a3a">
+      <EmbedInteropTypes>True</EmbedInteropTypes>
+    </Reference>
+    <Reference Include="WindowsBase" />
+  </ItemGroup>
+  <ItemGroup>
+    <COMReference Include="Microsoft.VisualStudio.CommandBars">
+      <Guid>{1CBA492E-7263-47BB-87FE-639000619B15}</Guid>
+      <VersionMajor>8</VersionMajor>
+      <VersionMinor>0</VersionMinor>
+      <Lcid>0</Lcid>
+      <WrapperTool>primary</WrapperTool>
+      <Isolated>False</Isolated>
+      <EmbedInteropTypes>False</EmbedInteropTypes>
+    </COMReference>
+    <COMReference Include="stdole">
+      <Guid>{00020430-0000-0000-C000-000000000046}</Guid>
+      <VersionMajor>2</VersionMajor>
+      <VersionMinor>0</VersionMinor>
+      <Lcid>0</Lcid>
+      <WrapperTool>primary</WrapperTool>
+      <Isolated>False</Isolated>
+      <EmbedInteropTypes>False</EmbedInteropTypes>
+    </COMReference>
+  </ItemGroup>
+  <ItemGroup>
+    <Compile Include="BatchBuilder.cs" />
+    <Compile Include="BatchBuilderToolControl.xaml.cs">
+      <DependentUpon>BatchBuilderToolControl.xaml</DependentUpon>
+    </Compile>
+    <Compile Include="BatchBuilderToolWindow.cs" />
+    <Compile Include="BuildStartupProject.cs" />
+    <Compile Include="BuildStatusIcon.xaml.cs">
+      <DependentUpon>BuildStatusIcon.xaml</DependentUpon>
+    </Compile>
+    <Compile Include="CommandLineEditor.cs" />
+    <Compile Include="CompileSingleFile.cs" />
+    <Compile Include="GenerateProjectFiles.cs" />
+    <Compile Include="Guids.cs" />
+    <Compile Include="Logging.cs" />
+    <Compile Include="ProfileManager.cs">
+      <SubType>Component</SubType>
+    </Compile>
+    <Compile Include="QuickBuild.cs" />
+    <Compile Include="Resources.Designer.cs">
+      <AutoGen>True</AutoGen>
+      <DesignTime>True</DesignTime>
+      <DependentUpon>Resources.resx</DependentUpon>
+    </Compile>
+    <Compile Include="GlobalSuppressions.cs" />
+    <Compile Include="StartupProjectSelector.cs" />
+    <Compile Include="UnrealVSOptions.cs">
+      <SubType>Component</SubType>
+    </Compile>
+    <Compile Include="UnrealVSPackage.cs" />
+    <Compile Include="Properties\AssemblyInfo.cs" />
+    <Compile Include="Utils.cs" />
+  </ItemGroup>
+  <ItemGroup>
+    <EmbeddedResource Include="Resources.resx">
+      <Generator>ResXFileCodeGenerator</Generator>
+      <LastGenOutput>Resources.Designer.cs</LastGenOutput>
+      <SubType>Designer</SubType>
+    </EmbeddedResource>
+    <EmbeddedResource Include="VSPackage.resx">
+      <MergeWithCTO>true</MergeWithCTO>
+      <ManifestResourceName>VSPackage</ManifestResourceName>
+      <SubType>Designer</SubType>
+    </EmbeddedResource>
+  </ItemGroup>
+  <ItemGroup>
+    <None Include="packages.config" />
+  </ItemGroup>
+  <ItemGroup>
+    <None Include="Key.snk" />
+  </ItemGroup>
+  <ItemGroup>
+    <Content Include="TargetManifests\VS14.source.extension.vsixmanifest" />
+    <Content Include="TargetManifests\VS14d.source.extension.vsixmanifest" />
+    <Content Include="TargetManifests\VS15.source.extension.vsixmanifest" />
+    <Content Include="TargetManifests\VS15d.source.extension.vsixmanifest" />
+    <Content Include="TargetManifests\VS16.source.extension.vsixmanifest" />
+    <Content Include="TargetManifests\VS16d.source.extension.vsixmanifest" />
+    <None Include="Generated\source.extension.vsixmanifest" />
+    <None Include="Resources\BatchBuilder.png" />
+    <Resource Include="Resources\BatchBuilderBackground.png" />
+    <Content Include="Resources\Icon.png">
+      <CopyToOutputDirectory>Always</CopyToOutputDirectory>
+      <IncludeInVSIX>true</IncludeInVSIX>
+    </Content>
+    <Content Include="Resources\Preview.png">
+      <IncludeInVSIX>true</IncludeInVSIX>
+      <CopyToOutputDirectory>Always</CopyToOutputDirectory>
+    </Content>
+    <Resource Include="Resources\RefreshProjects.png" />
+  </ItemGroup>
+  <ItemGroup>
+    <BootstrapperPackage Include=".NETFramework,Version=v4.5">
+      <Visible>False</Visible>
+      <ProductName>Microsoft .NET Framework 4.5 %28x86 and x64%29</ProductName>
+      <Install>true</Install>
+    </BootstrapperPackage>
+    <BootstrapperPackage Include="Microsoft.Net.Client.3.5">
+      <Visible>False</Visible>
+      <ProductName>.NET Framework 3.5 SP1 Client Profile</ProductName>
+      <Install>false</Install>
+    </BootstrapperPackage>
+    <BootstrapperPackage Include="Microsoft.Net.Framework.3.5.SP1">
+      <Visible>False</Visible>
+      <ProductName>.NET Framework 3.5 SP1</ProductName>
+      <Install>false</Install>
+    </BootstrapperPackage>
+  </ItemGroup>
+  <ItemGroup>
+    <Page Include="BatchBuilderToolControl.xaml">
+      <SubType>Designer</SubType>
+      <Generator>MSBuild:Compile</Generator>
+      <Generator>MSBuild:Compile</Generator>
+      <SubType>Designer</SubType>
+      <Generator>MSBuild:Compile</Generator>
+      <SubType>Designer</SubType>
+    </Page>
+    <Page Include="BuildStatusIcon.xaml">
+      <SubType>Designer</SubType>
+      <Generator>MSBuild:Compile</Generator>
+      <Generator>MSBuild:Compile</Generator>
+      <SubType>Designer</SubType>
+      <Generator>MSBuild:Compile</Generator>
+      <SubType>Designer</SubType>
+    </Page>
+  </ItemGroup>
+  <PropertyGroup>
+    <UseCodebase>true</UseCodebase>
+  </PropertyGroup>
+  <PropertyGroup>
+    <VisualStudioVersion Condition="'$(VisualStudioVersion)' == ''">11.0</VisualStudioVersion>
+    <VSToolsPath Condition="'$(VSToolsPath)' == ''">$(MSBuildExtensionsPath32)\Microsoft\VisualStudio\v$(VisualStudioVersion)</VSToolsPath>
+  </PropertyGroup>
+  <PropertyGroup Condition="'$(Configuration)|$(Platform)' == 'VS11|AnyCPU'">
+    <DebugType>pdbonly</DebugType>
+    <Optimize>true</Optimize>
+    <OutputPath>..\..\..\Extras\UnrealVS\VS2012\</OutputPath>
+    <DefineConstants>TRACE;CODE_ANALYSIS;VS11</DefineConstants>
+    <ErrorReport>prompt</ErrorReport>
+    <WarningLevel>4</WarningLevel>
+    <RunCodeAnalysis>true</RunCodeAnalysis>
+    <Prefer32Bit>false</Prefer32Bit>
+  </PropertyGroup>
+  <PropertyGroup Condition=" '$(Configuration)|$(Platform)' == 'VS11d|AnyCPU' ">
+    <DebugSymbols>true</DebugSymbols>
+    <DebugType>full</DebugType>
+    <Optimize>false</Optimize>
+    <OutputPath>bin\VS11d\</OutputPath>
+    <DefineConstants>TRACE;DEBUG;VS11</DefineConstants>
+    <ErrorReport>prompt</ErrorReport>
+    <WarningLevel>4</WarningLevel>
+    <Prefer32Bit>false</Prefer32Bit>
+    <CopyVsixExtensionFiles>False</CopyVsixExtensionFiles>
+    <DeployExtension>True</DeployExtension>
+    <CreateVsixContainer>True</CreateVsixContainer>
+  </PropertyGroup>
+  <PropertyGroup Condition=" '$(Configuration)|$(Platform)' == 'VS14d|AnyCPU' ">
+    <DebugSymbols>true</DebugSymbols>
+    <OutputPath>bin\VS14d\</OutputPath>
+    <DefineConstants>TRACE;DEBUG;VS14</DefineConstants>
+    <DebugType>full</DebugType>
+    <PlatformTarget>AnyCPU</PlatformTarget>
+    <ErrorReport>prompt</ErrorReport>
+    <CodeAnalysisRuleSet>MinimumRecommendedRules.ruleset</CodeAnalysisRuleSet>
+  </PropertyGroup>
+  <PropertyGroup Condition=" '$(Configuration)|$(Platform)' == 'VS14|AnyCPU' ">
+    <OutputPath>..\..\..\Extras\UnrealVS\VS2015\</OutputPath>
+    <DefineConstants>TRACE;CODE_ANALYSIS;VS14</DefineConstants>
+    <Optimize>true</Optimize>
+    <DebugType>pdbonly</DebugType>
+    <PlatformTarget>AnyCPU</PlatformTarget>
+    <RunCodeAnalysis>true</RunCodeAnalysis>
+    <ErrorReport>prompt</ErrorReport>
+    <CodeAnalysisRuleSet>MinimumRecommendedRules.ruleset</CodeAnalysisRuleSet>
+  </PropertyGroup>
+  <PropertyGroup Condition="'$(Configuration)|$(Platform)' == 'VS15|AnyCPU'">
+    <OutputPath>..\..\..\Extras\UnrealVS\VS2017\</OutputPath>
+    <DefineConstants>CODE_ANALYSIS;TRACE;CODE_ANALYSIS;VS14</DefineConstants>
+    <Optimize>true</Optimize>
+    <DebugType>pdbonly</DebugType>
+    <PlatformTarget>AnyCPU</PlatformTarget>
+    <RunCodeAnalysis>true</RunCodeAnalysis>
+    <ErrorReport>prompt</ErrorReport>
+    <CodeAnalysisRuleSet>MinimumRecommendedRules.ruleset</CodeAnalysisRuleSet>
+  </PropertyGroup>
+  <PropertyGroup Condition="'$(Configuration)|$(Platform)' == 'VS15d|AnyCPU'">
+    <DebugSymbols>true</DebugSymbols>
+    <OutputPath>bin\VS15d\</OutputPath>
+    <DefineConstants>TRACE;DEBUG;VS14</DefineConstants>
+    <DebugType>full</DebugType>
+    <PlatformTarget>AnyCPU</PlatformTarget>
+    <ErrorReport>prompt</ErrorReport>
+    <CodeAnalysisRuleSet>MinimumRecommendedRules.ruleset</CodeAnalysisRuleSet>
+    <DeployExtension>True</DeployExtension>
+    <CopyVsixExtensionFiles>False</CopyVsixExtensionFiles>
+  </PropertyGroup>
+  <PropertyGroup Condition=" '$(Configuration)|$(Platform)' == 'VS16d|AnyCPU' ">
+    <DebugSymbols>true</DebugSymbols>
+    <OutputPath>..\..\..\Extras\UnrealVS\VS2019\</OutputPath>
+    <DefineConstants>TRACE;DEBUG;VS14</DefineConstants>
+    <DebugType>full</DebugType>
+    <PlatformTarget>AnyCPU</PlatformTarget>
+    <ErrorReport>prompt</ErrorReport>
+    <CodeAnalysisRuleSet>MinimumRecommendedRules.ruleset</CodeAnalysisRuleSet>
+  </PropertyGroup>
+  <PropertyGroup Condition=" '$(Configuration)|$(Platform)' == 'VS16|AnyCPU' ">
+    <OutputPath>..\..\..\Extras\UnrealVS\VS2019\</OutputPath>
+    <DefineConstants>CODE_ANALYSIS;CODE_ANALYSIS;TRACE;CODE_ANALYSIS;VS14</DefineConstants>
+    <Optimize>true</Optimize>
+    <DebugType>pdbonly</DebugType>
+    <PlatformTarget>AnyCPU</PlatformTarget>
+    <RunCodeAnalysis>false</RunCodeAnalysis>
+    <ErrorReport>prompt</ErrorReport>
+    <CodeAnalysisRuleSet>MinimumRecommendedRules.ruleset</CodeAnalysisRuleSet>
+  </PropertyGroup>
+  <Import Project="$(MSBuildBinPath)\Microsoft.CSharp.targets" />
+  <Import Project="$(VSToolsPath)\VSSDK\Microsoft.VsSDK.targets" Condition="'$(VSToolsPath)' != ''" />
+  <Import Project="$(MSBuildExtensionsPath32)\Microsoft\VisualStudio\v11.0\VSSDK\Microsoft.VsSDK.targets" Condition="false" />
+  <PropertyGroup>
+    <PostBuildEvent>
+    </PostBuildEvent>
+  </PropertyGroup>
+  <PropertyGroup>
+    <PreBuildEvent>copy "$(ProjectDir)TargetManifests\$(ConfigurationName).source.extension.vsixmanifest" "$(ProjectDir)Generated\source.extension.vsixmanifest" /Y</PreBuildEvent>
+  </PropertyGroup>
+  <!-- To modify your build process, add your task inside one of the targets below and uncomment it. 
+       Other similar extension points exist, see Microsoft.Common.targets.
+  <Target Name="BeforeBuild">
+  </Target>
+  <Target Name="AfterBuild">
+  </Target>
+  -->
+</Project>
\ No newline at end of file
