UnrealVS now loads asynchronously and no longer pops up a warning about it on startup.

Updated to version to 1.55

(ushell-p4-cherrypick of 7288939 by sebastian.nordgren)

diff --git a/Engine/Source/Programs/UnrealVS/BuildStartupProject.cs b/Engine/Source/Programs/UnrealVS/BuildStartupProject.cs
--- a/Engine/Source/Programs/UnrealVS/BuildStartupProject.cs
+++ b/Engine/Source/Programs/UnrealVS/BuildStartupProject.cs
@@ -1,4 +1,4 @@
-﻿// Copyright 1998-2019 Epic Games, Inc. All Rights Reserved.
+// Copyright 1998-2019 Epic Games, Inc. All Rights Reserved.
 
 using System;
 using System.Diagnostics;
@@ -12,7 +12,6 @@ using System.Runtime.InteropServices;
 using System.Collections.Generic;
 using EnvDTE80;
 
-
 namespace UnrealVS
 {
 	class BuildStartupProject
diff --git a/Engine/Source/Programs/UnrealVS/Properties/AssemblyInfo.cs b/Engine/Source/Programs/UnrealVS/Properties/AssemblyInfo.cs
--- a/Engine/Source/Programs/UnrealVS/Properties/AssemblyInfo.cs
+++ b/Engine/Source/Programs/UnrealVS/Properties/AssemblyInfo.cs
@@ -31,8 +31,8 @@ using System.Runtime.InteropServices;
 // You can specify all the values or you can default the Revision and Build Numbers 
 // by using the '*' as shown below:
 
-[assembly: AssemblyVersion("1.54.0.0")]
-[assembly: AssemblyFileVersion("1.54.0.0")]
+[assembly: AssemblyVersion("1.55.0.0")]
+[assembly: AssemblyFileVersion("1.55.0.0")]
 
 
 
diff --git a/Engine/Source/Programs/UnrealVS/UnrealVSPackage.cs b/Engine/Source/Programs/UnrealVS/UnrealVSPackage.cs
--- a/Engine/Source/Programs/UnrealVS/UnrealVSPackage.cs
+++ b/Engine/Source/Programs/UnrealVS/UnrealVSPackage.cs
@@ -3,13 +3,11 @@
 using System;
 using System.Collections.Generic;
 using System.Diagnostics;
-using System.Globalization;
 using System.IO;
 using System.Linq;
 using System.Runtime.InteropServices;
 using System.ComponentModel.Design;
 using System.Threading;
-using Microsoft.Win32;
 using Microsoft.VisualStudio;
 using Microsoft.VisualStudio.Shell.Interop;
 using Microsoft.VisualStudio.OLE.Interop;
@@ -21,7 +19,7 @@ using Thread = System.Threading.Thread;
 namespace UnrealVS
 {
 	// This attribute tells the PkgDef creation utility (CreatePkgDef.exe) that this class is a VS package.
-	[PackageRegistration( UseManagedResourcesOnly = true )]
+	[PackageRegistration( UseManagedResourcesOnly = true, AllowsBackgroundLoading = true )]
 
 	// This attribute is used to register the informations needed to show the this package in the Help/About dialog of Visual Studio.
 	[InstalledProductRegistration("#110", "#112", VersionString, IconResourceID = 400)]
@@ -31,7 +29,7 @@ namespace UnrealVS
 	[ProvideMenuResource( "Menus.ctmenu", 1 )]
 
 	// Force the package to load whenever a solution exists
-	[ProvideAutoLoad(UIContextGuids80.SolutionExists)]
+	[ProvideAutoLoad(UIContextGuids80.SolutionExists, PackageAutoLoadFlags.BackgroundLoad)]
 
 	// Register the settings implementing class as providing support to UnrealVSPackage.
 	//[ProvideProfile(typeof (ProfileManager), "UnrealVS", "UnrealVSPackage", 110, 113, false)]
@@ -50,7 +48,7 @@ namespace UnrealVS
 	/// with Visual Studio shell and serves as the entry point into our extension
 	/// </summary>
 	public sealed class UnrealVSPackage :
-		Package,		// We inherit from Package which allows us to become a "plugin" within the Visual Studio shell
+		AsyncPackage,			// We inherit from AsyncPackage which allows us to become a "plugin" within the Visual Studio shell
 		IVsSolutionEvents,		// This interface allows us to register to be notified of events such as opening a project
 		IVsUpdateSolutionEvents,// Allows us to register to be notified of events such as active config changes
 		IVsSelectionEvents,		// Allows us to be notified when the startup project has changed to a different project
@@ -59,7 +57,7 @@ namespace UnrealVS
 	{
 		/** Constants */
 
-		private const string VersionString = "v1.53";
+		private const string VersionString = "v1.55";
 		private const string UnrealSolutionFileNamePrefix = "UE4";
 		private const string ExtensionName = "UnrealVS";
 		private const string CommandLineOptionKey = ExtensionName + "CommandLineMRU";
@@ -182,33 +180,35 @@ namespace UnrealVS
 		/// <summary>
 		/// Initializes the package right after it's been "sited" into the fully-initialized Visual Studio IDE.
 		/// </summary>
-		protected override void Initialize()
+		protected override async System.Threading.Tasks.Task InitializeAsync(CancellationToken cancellationToken, IProgress<ServiceProgressData> progress)
 		{
 			Logging.WriteLine("Initializing UnrealVS extension...");
 
+			await JoinableTaskFactory.SwitchToMainThreadAsync(cancellationToken);
+
 			// Grab the MenuCommandService
-			MenuCommandService = GetService( typeof( IMenuCommandService ) ) as OleMenuCommandService;
+			MenuCommandService = await GetServiceAsync( typeof( IMenuCommandService ) ) as OleMenuCommandService;
 			
 			// Get access to Visual Studio's DTE object.  This object has various hooks into the Visual Studio
 			// shell that are useful for writing extensions.
-			DTE = (DTE)GetGlobalService( typeof( DTE ) );
+			DTE = await GetServiceAsync(typeof(DTE)) as DTE;
 			Logging.WriteLine("DTE version " + DTE.Version);
 
 			// Get selection manager and register to receive events
 			SelectionManager =
-				ServiceProvider.GlobalProvider.GetService(typeof (SVsShellMonitorSelection)) as IVsMonitorSelection;
-			SelectionManager.AdviseSelectionEvents( this, out SelectionEventsHandle );
+				await GetServiceAsync(typeof(SVsShellMonitorSelection)) as IVsMonitorSelection;
+			SelectionManager.AdviseSelectionEvents(this, out SelectionEventsHandle);
 
 			// Get solution and register to receive events
-			SolutionManager = ServiceProvider.GlobalProvider.GetService( typeof( SVsSolution ) ) as IVsSolution2;
+			SolutionManager = await GetServiceAsync(typeof(SVsSolution)) as IVsSolution2;
 			UpdateUnrealLoadedStatus();
-			SolutionManager.AdviseSolutionEvents( this, out SolutionEventsHandle );
+			SolutionManager.AdviseSolutionEvents(this, out SolutionEventsHandle);
 
 			// Grab the solution build manager.  We need this in order to change certain things about the Visual
 			// Studio environment, like what the active startup project is
 			// Get solution build manager
 			SolutionBuildManager =
-				ServiceProvider.GlobalProvider.GetService(typeof (SVsSolutionBuildManager)) as IVsSolutionBuildManager2;
+				await GetServiceAsync(typeof(SVsSolutionBuildManager)) as IVsSolutionBuildManager2;
 			SolutionBuildManager.AdviseUpdateSolutionEvents(this, out UpdateSolutionEventsHandle);
 
 			// Create our command-line editor
diff --git a/Engine/Source/Programs/UnrealVS/UnrealVS_2019.csproj b/Engine/Source/Programs/UnrealVS/UnrealVS_2019.csproj
--- a/Engine/Source/Programs/UnrealVS/UnrealVS_2019.csproj
+++ b/Engine/Source/Programs/UnrealVS/UnrealVS_2019.csproj
@@ -85,19 +85,24 @@
     <Reference Include="Microsoft.VisualStudio.OLE.Interop, Version=7.1.40304.0, Culture=neutral, PublicKeyToken=b03f5f7f11d50a3a" />
     <Reference Include="Microsoft.VisualStudio.Shell.15.0, Version=16.0.0.0, Culture=neutral, PublicKeyToken=b03f5f7f11d50a3a, processorArchitecture=MSIL" />
     <Reference Include="Microsoft.VisualStudio.Shell.Framework, Version=16.0.0.0, Culture=neutral, PublicKeyToken=b03f5f7f11d50a3a, processorArchitecture=MSIL" />
+    <Reference Include="Microsoft.VisualStudio.Shell.Immutable.14.0, Version=15.0.0.0, Culture=neutral, PublicKeyToken=b03f5f7f11d50a3a, processorArchitecture=MSIL" />
     <Reference Include="Microsoft.VisualStudio.Shell.Interop, Version=7.1.40304.0, Culture=neutral, PublicKeyToken=b03f5f7f11d50a3a" />
     <Reference Include="Microsoft.VisualStudio.Shell.Interop.10.0, Version=10.0.0.0, Culture=neutral, PublicKeyToken=b03f5f7f11d50a3a, processorArchitecture=MSIL">
       <EmbedInteropTypes>True</EmbedInteropTypes>
     </Reference>
-    <Reference Include="Microsoft.VisualStudio.Shell.Interop.12.0, Version=12.0.0.0, Culture=neutral, PublicKeyToken=b03f5f7f11d50a3a, processorArchitecture=MSIL">
+    <Reference Include="Microsoft.VisualStudio.Shell.Interop.14.0.DesignTime, Version=14.0.0.0, Culture=neutral, PublicKeyToken=b03f5f7f11d50a3a, processorArchitecture=MSIL">
       <EmbedInteropTypes>True</EmbedInteropTypes>
     </Reference>
     <Reference Include="Microsoft.VisualStudio.Shell.Interop.15.3.DesignTime, Version=15.3.0.0, Culture=neutral, PublicKeyToken=b03f5f7f11d50a3a, processorArchitecture=MSIL">
       <EmbedInteropTypes>True</EmbedInteropTypes>
     </Reference>
+    <Reference Include="Microsoft.VisualStudio.Shell.Interop.15.6.DesignTime, Version=15.6.0.0, Culture=neutral, PublicKeyToken=b03f5f7f11d50a3a, processorArchitecture=MSIL">
+      <EmbedInteropTypes>True</EmbedInteropTypes>
+    </Reference>
     <Reference Include="Microsoft.VisualStudio.Shell.Interop.8.0, Version=8.0.0.0, Culture=neutral, PublicKeyToken=b03f5f7f11d50a3a" />
     <Reference Include="Microsoft.VisualStudio.Shell.Interop.9.0, Version=9.0.0.0, Culture=neutral, PublicKeyToken=b03f5f7f11d50a3a" />
     <Reference Include="Microsoft.VisualStudio.TextManager.Interop, Version=7.1.40304.0, Culture=neutral, PublicKeyToken=b03f5f7f11d50a3a" />
+    <Reference Include="Microsoft.VisualStudio.Threading, Version=16.0.0.0, Culture=neutral, PublicKeyToken=b03f5f7f11d50a3a" />
     <Reference Include="Microsoft.VisualStudio.VCProject, Version=16.0.0.0, Culture=neutral, PublicKeyToken=b03f5f7f11d50a3a">
       <EmbedInteropTypes>True</EmbedInteropTypes>
     </Reference>
@@ -123,15 +128,6 @@
     <Reference Include="WindowsBase" />
   </ItemGroup>
   <ItemGroup>
-    <COMReference Include="Microsoft.VisualStudio.CommandBars">
-      <Guid>{1CBA492E-7263-47BB-87FE-639000619B15}</Guid>
-      <VersionMajor>8</VersionMajor>
-      <VersionMinor>0</VersionMinor>
-      <Lcid>0</Lcid>
-      <WrapperTool>primary</WrapperTool>
-      <Isolated>False</Isolated>
-      <EmbedInteropTypes>False</EmbedInteropTypes>
-    </COMReference>
     <COMReference Include="stdole">
       <Guid>{00020430-0000-0000-C000-000000000046}</Guid>
       <VersionMajor>2</VersionMajor>
@@ -336,6 +332,7 @@
     <RunCodeAnalysis>false</RunCodeAnalysis>
     <ErrorReport>prompt</ErrorReport>
     <CodeAnalysisRuleSet>MinimumRecommendedRules.ruleset</CodeAnalysisRuleSet>
+    <CopyVsixExtensionFiles>False</CopyVsixExtensionFiles>
   </PropertyGroup>
   <Import Project="$(MSBuildBinPath)\Microsoft.CSharp.targets" />
   <Import Project="$(VSToolsPath)\VSSDK\Microsoft.VsSDK.targets" Condition="'$(VSToolsPath)' != ''" />
