UnrealVS: Fix single-file compile not working for foreign projects where the engine folder contains a space.

diff --git a/Engine/Source/Programs/UnrealVS/CompileSingleFile.cs b/Engine/Source/Programs/UnrealVS/CompileSingleFile.cs
--- a/Engine/Source/Programs/UnrealVS/CompileSingleFile.cs
+++ b/Engine/Source/Programs/UnrealVS/CompileSingleFile.cs
@@ -1,4 +1,4 @@
-﻿// Copyright 1998-2018 Epic Games, Inc. All Rights Reserved.
+// Copyright 1998-2018 Epic Games, Inc. All Rights Reserved.
 
 using System;
 using System.Diagnostics;
@@ -177,7 +177,7 @@ namespace UnrealVS
 			// Spawn the new process
 			ChildProcess = new System.Diagnostics.Process();
 			ChildProcess.StartInfo.FileName = Path.Combine(Environment.SystemDirectory, "cmd.exe");
-			ChildProcess.StartInfo.Arguments = String.Format("/C {0} -singlefile=\"{1}\"", BuildCommandLine, FileToCompile);
+			ChildProcess.StartInfo.Arguments = String.Format("/C \"{0} -singlefile=\"{1}\"\"", BuildCommandLine, FileToCompile);
 			ChildProcess.StartInfo.WorkingDirectory = Path.GetDirectoryName(StartupProject.FullName);
 			ChildProcess.StartInfo.UseShellExecute = false;
 			ChildProcess.StartInfo.RedirectStandardOutput = true;
diff --git a/Engine/Source/Programs/UnrealVS/Properties/AssemblyInfo.cs b/Engine/Source/Programs/UnrealVS/Properties/AssemblyInfo.cs
--- a/Engine/Source/Programs/UnrealVS/Properties/AssemblyInfo.cs
+++ b/Engine/Source/Programs/UnrealVS/Properties/AssemblyInfo.cs
@@ -1,4 +1,4 @@
-﻿// Copyright 1998-2018 Epic Games, Inc. All Rights Reserved.
+// Copyright 1998-2018 Epic Games, Inc. All Rights Reserved.
 
 using System;
 using System.Reflection;
@@ -31,8 +31,8 @@ using System.Runtime.InteropServices;
 // You can specify all the values or you can default the Revision and Build Numbers 
 // by using the '*' as shown below:
 
-[assembly: AssemblyVersion("1.51.0.0")]
-[assembly: AssemblyFileVersion("1.51.0.0")]
+[assembly: AssemblyVersion("1.52.0.0")]
+[assembly: AssemblyFileVersion("1.52.0.0")]
 
 
 
diff --git a/Engine/Source/Programs/UnrealVS/UnrealVSPackage.cs b/Engine/Source/Programs/UnrealVS/UnrealVSPackage.cs
--- a/Engine/Source/Programs/UnrealVS/UnrealVSPackage.cs
+++ b/Engine/Source/Programs/UnrealVS/UnrealVSPackage.cs
@@ -59,7 +59,7 @@ namespace UnrealVS
 	{
 		/** Constants */
 
-		private const string VersionString = "v1.51";
+		private const string VersionString = "v1.52";
 		private const string UnrealSolutionFileNamePrefix = "UE4";
 		private const string ExtensionName = "UnrealVS";
 		private const string CommandLineOptionKey = ExtensionName + "CommandLineMRU";
