Addresed some races in UnrealVS caused by it now loading async.
Fixed Quick Build menu being empty sometimes.
Fixed issue loading .suo settings

Bumped version to 1.56

diff --git a/Engine/Source/Programs/UnrealVS/BatchBuilder.cs b/Engine/Source/Programs/UnrealVS/BatchBuilder.cs
--- a/Engine/Source/Programs/UnrealVS/BatchBuilder.cs
+++ b/Engine/Source/Programs/UnrealVS/BatchBuilder.cs
@@ -1,4 +1,4 @@
-﻿// Copyright 1998-2019 Epic Games, Inc. All Rights Reserved.
+// Copyright 1998-2019 Epic Games, Inc. All Rights Reserved.
 
 using System;
 using System.Collections.Generic;
@@ -28,12 +28,13 @@ namespace UnrealVS
 			}
 		}
 
+		public static BatchBuilderToolControl.BatchBuilderToolState _ToolControlState = new BatchBuilderToolControl.BatchBuilderToolState();
 		private static BatchBuilderToolControl _ToolControl;
 		public static BatchBuilderToolControl ToolControl
 		{
 			get 
 			{
-				if (_ToolControl == null) _ToolControl = new BatchBuilderToolControl();
+				if (_ToolControl == null) _ToolControl = new BatchBuilderToolControl(_ToolControlState);
 				return _ToolControl;
 			}
 		}
@@ -41,6 +42,10 @@ namespace UnrealVS
 		/** methods */
 
 		public BatchBuilder()
+		{
+		}
+
+		public void Initialize()
 		{
 			// Create the command for the tool window
 			var ToolWindowCommandId = new CommandID(GuidList.UnrealVSCmdSet, BatchBuilderToolWindowId);
@@ -54,7 +59,7 @@ namespace UnrealVS
 		/// <param name="Stream">The stream to load the option data from.</param>
 		public void LoadOptions(Stream Stream)
 		{
-			ToolControl.LoadOptions(Stream);
+			_ToolControlState.LoadOptions(Stream);
 		}
 
 		/// <summary>
@@ -63,7 +68,7 @@ namespace UnrealVS
 		/// <param name="Stream">The stream to save the option data to.</param>
 		public void SaveOptions(Stream Stream)
 		{
-			ToolControl.SaveOptions(Stream);
+			_ToolControlState.SaveOptions(Stream);
 		}
 
 		/// <summary>
diff --git a/Engine/Source/Programs/UnrealVS/BatchBuilderToolControl.xaml.cs b/Engine/Source/Programs/UnrealVS/BatchBuilderToolControl.xaml.cs
--- a/Engine/Source/Programs/UnrealVS/BatchBuilderToolControl.xaml.cs
+++ b/Engine/Source/Programs/UnrealVS/BatchBuilderToolControl.xaml.cs
@@ -1,8 +1,9 @@
-﻿// Copyright 1998-2019 Epic Games, Inc. All Rights Reserved.
+// Copyright 1998-2019 Epic Games, Inc. All Rights Reserved.
 
 using System;
 using System.Collections.Generic;
 using System.Collections.ObjectModel;
+using System.Collections.Specialized;
 using System.IO;
 using System.Linq;
 using System.Text;
@@ -14,6 +15,7 @@ using EnvDTE80;
 using Microsoft.VisualStudio;
 using Microsoft.VisualStudio.Shell;
 using Microsoft.VisualStudio.Shell.Interop;
+using Constants = EnvDTE.Constants;
 
 namespace UnrealVS
 {
@@ -22,6 +24,7 @@ namespace UnrealVS
 	/// </summary>
 	public partial class BatchBuilderToolControl : UserControl
 	{
+		private readonly BatchBuilderToolState _state;
 		////////////// Types ///////////////////////
 
 		public class BuildJob : DependencyObject
@@ -230,6 +233,100 @@ namespace UnrealVS
 				return "<None>";
 			}
 		}
+		
+		public class BatchBuilderToolState
+		{
+			public readonly ObservableCollection<BuildJobSet> _BuildJobSetsCollection = new ObservableCollection<BuildJobSet>();
+
+			/// <summary>
+			/// Called from the package class when there are options to be read out of the solution file.
+			/// </summary>
+			/// <param name="Stream">The stream to load the option data from.</param>
+			public void LoadOptions(Stream Stream)
+			{
+				try
+				{
+					_BuildJobSetsCollection.Clear();
+
+					using (BinaryReader Reader = new BinaryReader(Stream))
+					{
+						int SetCount = Reader.ReadInt32();
+
+						for (int SetIdx = 0; SetIdx < SetCount; SetIdx++)
+						{
+							BuildJobSet LoadedSet = new BuildJobSet();
+							LoadedSet.Name = Reader.ReadString();
+							int JobCount = Reader.ReadInt32();
+							for (int JobIdx = 0; JobIdx < JobCount; JobIdx++)
+							{
+								Utils.SafeProjectReference ProjectRef = new Utils.SafeProjectReference { FullName = Reader.ReadString(), Name = Reader.ReadString() };
+
+								string Config = Reader.ReadString();
+								string Platform = Reader.ReadString();
+								BuildJob.BuildJobType JobType;
+
+								if (Enum.TryParse(Reader.ReadString(), out JobType))
+								{
+									LoadedSet.BuildJobs.Add(new BuildJob(ProjectRef, Config, Platform, JobType));
+								}
+							}
+							_BuildJobSetsCollection.Add(LoadedSet);
+						}
+					}
+
+					if (_BuildJobSetsCollection.Count == 0)
+					{
+						BuildJobSet DefaultItem = new BuildJobSet { Name = "UntitledSet" };
+						_BuildJobSetsCollection.Add(DefaultItem);
+					}
+
+					StateLoaded?.Invoke(this, null);
+				}
+				catch (Exception ex)
+				{
+					Exception AppEx = new ApplicationException("BatchBuilder failed to load options from .suo", ex);
+					Logging.WriteLine(AppEx.ToString());
+					throw AppEx;
+				}
+			}
+
+			/// <summary>
+			/// Called from the package class when there are options to be written to the solution file.
+			/// </summary>
+			/// <param name="Stream">The stream to save the option data to.</param>
+			public void SaveOptions(Stream Stream)
+			{
+				try
+				{
+					using (BinaryWriter Writer = new BinaryWriter(Stream))
+					{
+						Writer.Write(_BuildJobSetsCollection.Count);
+						foreach (var Set in _BuildJobSetsCollection)
+						{
+							Writer.Write(Set.Name);
+							Writer.Write(Set.BuildJobs.Count);
+							foreach (var Job in Set.BuildJobs)
+							{
+								Writer.Write(Job.Project.FullName);
+								Writer.Write(Job.Project.Name);
+								Writer.Write(Job.Config);
+								Writer.Write(Job.Platform);
+								Writer.Write(Enum.GetName(typeof(BuildJob.BuildJobType), Job.JobType) ?? "INVALIDJOBTYPE");
+							}
+						}
+					}
+				}
+				catch (Exception ex)
+				{
+					Exception AppEx = new ApplicationException("BatchBuilder failed to save options to .suo", ex);
+					Logging.WriteLine(AppEx.ToString());
+					throw AppEx;
+				}
+			}
+
+			// we are using a custom even instead of listening to the observable collection as we just want a even when the entire collection has been modified
+			public event EventHandler StateLoaded;
+		}
 
 		////////////// Dependency properties ///////////////////////
 
@@ -317,7 +414,7 @@ namespace UnrealVS
 
 		public ObservableCollection<BuildJobSet> BuildJobSets
 		{
-			get { return _BuildJobSetsCollection; }
+			get { return _state._BuildJobSetsCollection; }
 		}
 
 		public string SelectedBuildJobSetName
@@ -327,14 +424,14 @@ namespace UnrealVS
 				if (String.IsNullOrEmpty(value)) return;
 
 				BuildJobSet SelectedSet =
-					_BuildJobSetsCollection.FirstOrDefault(
+					_state._BuildJobSetsCollection.FirstOrDefault(
 						Set => String.Compare(Set.Name, value, StringComparison.InvariantCultureIgnoreCase) == 0);
 
 				if (SelectedSet == null)
 				{
 					SelectedSet = new BuildJobSet { Name = value };
 					SelectedSet.DeepCopyJobsFrom(BuildJobs);
-					_BuildJobSetsCollection.Add(SelectedSet);
+					_state._BuildJobSetsCollection.Add(SelectedSet);
 				}
 
 				SetCombo.SelectedItem = SelectedSet;
@@ -407,8 +504,9 @@ namespace UnrealVS
 		/// <summary>
 		/// Default constructor
 		/// </summary>
-		public BatchBuilderToolControl()
+		public BatchBuilderToolControl(BatchBuilderToolState state)
 		{
+			_state = state;
 			IsSolutionOpen = UnrealVSPackage.Instance.DTE.Solution.IsOpen;
 
 			InitializeComponent();
@@ -435,89 +533,17 @@ namespace UnrealVS
 			BuildRadioButton.IsChecked = true;
 
 			EnsureDefaultBuildJobSet();
-		}
-
-		/// <summary>
-		/// Called from the package class when there are options to be read out of the solution file.
-		/// </summary>
-		/// <param name="Stream">The stream to load the option data from.</param>
-		public void LoadOptions(Stream Stream)
-		{
-			try
-			{
-				_BuildJobSetsCollection.Clear();
-
-				using (BinaryReader Reader = new BinaryReader(Stream))
-				{
-					int SetCount = Reader.ReadInt32();
-
-					for (int SetIdx = 0; SetIdx < SetCount; SetIdx++)
-					{
-						BuildJobSet LoadedSet = new BuildJobSet();
-						LoadedSet.Name = Reader.ReadString();
-						int JobCount = Reader.ReadInt32();
-						for (int JobIdx = 0; JobIdx < JobCount; JobIdx++)
-						{
-							Utils.SafeProjectReference ProjectRef = new Utils.SafeProjectReference { FullName = Reader.ReadString(), Name = Reader.ReadString() };
-
-							string Config = Reader.ReadString();
-							string Platform = Reader.ReadString();
-							BuildJob.BuildJobType JobType;
-
-							if (Enum.TryParse(Reader.ReadString(), out JobType))
-							{
-								LoadedSet.BuildJobs.Add(new BuildJob(ProjectRef, Config, Platform, JobType));
-							}
-						}
-						_BuildJobSetsCollection.Add(LoadedSet);
-					}
-				}
 
-				EnsureDefaultBuildJobSet();
-				if (SetCombo.SelectedItem == null)
-				{
-					SetCombo.SelectedItem = _BuildJobSetsCollection[0];
-				}
-			}
-			catch (Exception ex)
-			{
-				Exception AppEx = new ApplicationException("BatchBuilder failed to load options from .suo", ex);
-				Logging.WriteLine(AppEx.ToString());
-				throw AppEx;
-			}
+			state.StateLoaded += BuildJobSetsCollectionOnCollectionChanged;
 		}
 
-		/// <summary>
-		/// Called from the package class when there are options to be written to the solution file.
-		/// </summary>
-		/// <param name="Stream">The stream to save the option data to.</param>
-		public void SaveOptions(Stream Stream)
+		private void BuildJobSetsCollectionOnCollectionChanged(object sender, EventArgs eventArgs)
 		{
-			try
-			{
-				using (BinaryWriter Writer = new BinaryWriter(Stream))
-				{
-					Writer.Write(_BuildJobSetsCollection.Count);
-					foreach (var Set in _BuildJobSetsCollection)
-					{
-						Writer.Write(Set.Name);
-						Writer.Write(Set.BuildJobs.Count);
-						foreach (var Job in Set.BuildJobs)
-						{
-							Writer.Write(Job.Project.FullName);
-							Writer.Write(Job.Project.Name);
-							Writer.Write(Job.Config);
-							Writer.Write(Job.Platform);
-							Writer.Write(Enum.GetName(typeof(BuildJob.BuildJobType), Job.JobType) ?? "INVALIDJOBTYPE");
-						}
-					}
-				}
-			}
-			catch (Exception ex)
+			Logging.WriteLine("BuildJobSet Collection changed");
+
+			if (SetCombo.SelectedItem == null)
 			{
-				Exception AppEx = new ApplicationException("BatchBuilder failed to save options to .suo", ex);
-				Logging.WriteLine(AppEx.ToString());
-				throw AppEx;
+				SetCombo.SelectedItem = _state._BuildJobSetsCollection[0];
 			}
 		}
 
@@ -533,10 +559,10 @@ namespace UnrealVS
 
 		private void EnsureDefaultBuildJobSet()
 		{
-			if (_BuildJobSetsCollection.Count == 0)
+			if (_state._BuildJobSetsCollection.Count == 0)
 			{
 				BuildJobSet DefaultItem = new BuildJobSet {Name = "UntitledSet"};
-				_BuildJobSetsCollection.Add(DefaultItem);
+				_state._BuildJobSetsCollection.Add(DefaultItem);
 				if (SetCombo != null)
 				{
 					SetCombo.SelectedItem = DefaultItem;
@@ -566,7 +592,7 @@ namespace UnrealVS
 
 			if (!IsSolutionOpen)
 			{
-				_BuildJobSetsCollection.Clear();
+				_state._BuildJobSetsCollection.Clear();
 			}
 
 			string[] RefreshConfigs;
@@ -737,10 +763,10 @@ namespace UnrealVS
 			{
 				_LastSelectedBuildJobSet = (BuildJobSet) SetCombo.SelectedItem;
 				BuildJobs = _LastSelectedBuildJobSet.BuildJobs;
-				IsDeletableSetSelected = _BuildJobSetsCollection.Count > 1;
+				IsDeletableSetSelected = _state._BuildJobSetsCollection.Count > 1;
 				BuildJobsPanelTitle = String.Format("{0} ({1})", BuildJobsPanelPrefix, _LastSelectedBuildJobSet.Name);
 			}
-			if (!_BuildJobSetsCollection.Contains(_LastSelectedBuildJobSet))
+			if (!_state._BuildJobSetsCollection.Contains(_LastSelectedBuildJobSet))
 			{
 				_LastSelectedBuildJobSet = null;
 				BuildJobs = new ObservableCollection<BuildJob>();
@@ -761,12 +787,12 @@ namespace UnrealVS
 		{
 			if (_LastSelectedBuildJobSet != null)
 			{
-				int SelectionIdx = _BuildJobSetsCollection.IndexOf(_LastSelectedBuildJobSet);
-				if (SelectionIdx >= 0 && SelectionIdx < _BuildJobSetsCollection.Count)
+				int SelectionIdx = _state._BuildJobSetsCollection.IndexOf(_LastSelectedBuildJobSet);
+				if (SelectionIdx >= 0 && SelectionIdx < _state._BuildJobSetsCollection.Count)
 				{
 					SetCombo.SelectedItem = null;
-					_BuildJobSetsCollection.RemoveAt(SelectionIdx);
-					if (!(SelectionIdx < _BuildJobSetsCollection.Count))
+					_state._BuildJobSetsCollection.RemoveAt(SelectionIdx);
+					if (!(SelectionIdx < _state._BuildJobSetsCollection.Count))
 					{
 						SelectionIdx--;
 					}
@@ -850,7 +876,7 @@ namespace UnrealVS
 			string[] SolutionPlatforms;
 			Utils.GetSolutionConfigsAndPlatforms(out SolutionConfigs, out SolutionPlatforms);
 
-			foreach (var JobSet in _BuildJobSetsCollection)
+			foreach (var JobSet in _state._BuildJobSetsCollection)
 			{
 				for (int JobIdx = JobSet.BuildJobs.Count - 1; JobIdx >= 0; JobIdx--)
 				{
@@ -912,7 +938,7 @@ namespace UnrealVS
 		{
 			var OutputText = new StringBuilder();
 
-			var Window = UnrealVSPackage.Instance.DTE.Windows.Item(EnvDTE.Constants.vsWindowKindOutput);
+			var Window = UnrealVSPackage.Instance.DTE.Windows.Item(Constants.vsWindowKindOutput);
 			if (Window != null)
 			{
 				var OutputWindow = Window.Object as OutputWindow;
@@ -1019,7 +1045,5 @@ namespace UnrealVS
 		private readonly ObservableCollection<ProjectListItem> _ProjectCollection = new ObservableCollection<ProjectListItem>();
 		private readonly ObservableCollection<string> _ConfigCollection = new ObservableCollection<string>();
 		private readonly ObservableCollection<string> _PlatformCollection = new ObservableCollection<string>();
-
-		private readonly ObservableCollection<BuildJobSet> _BuildJobSetsCollection = new ObservableCollection<BuildJobSet>();
 	}
 }
\ No newline at end of file
diff --git a/Engine/Source/Programs/UnrealVS/CommandLineEditor.cs b/Engine/Source/Programs/UnrealVS/CommandLineEditor.cs
--- a/Engine/Source/Programs/UnrealVS/CommandLineEditor.cs
+++ b/Engine/Source/Programs/UnrealVS/CommandLineEditor.cs
@@ -28,6 +28,11 @@ namespace UnrealVS
 		/** methods */
 
 		public CommandLineEditor()
+		{
+		}
+
+
+		public void Initialize()
 		{
 			// Create the handlers for our commands
 			{
diff --git a/Engine/Source/Programs/UnrealVS/Properties/AssemblyInfo.cs b/Engine/Source/Programs/UnrealVS/Properties/AssemblyInfo.cs
--- a/Engine/Source/Programs/UnrealVS/Properties/AssemblyInfo.cs
+++ b/Engine/Source/Programs/UnrealVS/Properties/AssemblyInfo.cs
@@ -31,8 +31,8 @@ using System.Runtime.InteropServices;
 // You can specify all the values or you can default the Revision and Build Numbers 
 // by using the '*' as shown below:
 
-[assembly: AssemblyVersion("1.55.0.0")]
-[assembly: AssemblyFileVersion("1.55.0.0")]
+[assembly: AssemblyVersion("1.56.0.0")]
+[assembly: AssemblyFileVersion("1.56.0.0")]
 
 
 
diff --git a/Engine/Source/Programs/UnrealVS/QuickBuild.cs b/Engine/Source/Programs/UnrealVS/QuickBuild.cs
--- a/Engine/Source/Programs/UnrealVS/QuickBuild.cs
+++ b/Engine/Source/Programs/UnrealVS/QuickBuild.cs
@@ -1,4 +1,4 @@
-﻿// Copyright 1998-2019 Epic Games, Inc. All Rights Reserved.
+// Copyright 1998-2019 Epic Games, Inc. All Rights Reserved.
 
 using EnvDTE;
 using EnvDTE80;
@@ -257,6 +257,8 @@ namespace UnrealVS
 
 			UnrealVSPackage.Instance.OnUIContextChanged += OnUIContextChanged;
 			UnrealVSPackage.Instance.OnSolutionOpened += OnSolutionChanged;
+
+			OnSolutionChanged();
 		}
 
 		/// <summary>
@@ -297,6 +299,7 @@ namespace UnrealVS
 			CacheBuildConfigs();
 
 			ActivePlaformMenus.Clear();
+			Logging.WriteLine("Updating solution menu state, currently active config platforms are: " + string.Join(", ", SolutionConfigPlatforms));
 			foreach (var SubMenu in SubMenus)
 			{
 				if (SolutionConfigPlatforms.Any(Platform => string.Compare(Platform, SubMenu.Name, StringComparison.InvariantCultureIgnoreCase) == 0))
@@ -319,6 +322,7 @@ namespace UnrealVS
 
 			SolutionConfigNames =
 							(from SolutionConfiguration2 Sc in SolutionConfigs select Sc.Name).Distinct().ToArray();
+
 		}
 
 		/// <summary>
diff --git a/Engine/Source/Programs/UnrealVS/UnrealVSPackage.cs b/Engine/Source/Programs/UnrealVS/UnrealVSPackage.cs
--- a/Engine/Source/Programs/UnrealVS/UnrealVSPackage.cs
+++ b/Engine/Source/Programs/UnrealVS/UnrealVSPackage.cs
@@ -29,6 +29,7 @@ namespace UnrealVS
 	[ProvideMenuResource( "Menus.ctmenu", 1 )]
 
 	// Force the package to load whenever a solution exists
+	//[ProvideAutoLoad(VSConstants.UICONTEXT.SolutionExistsAndFullyLoaded_string, PackageAutoLoadFlags.BackgroundLoad)]
 	[ProvideAutoLoad(UIContextGuids80.SolutionExists, PackageAutoLoadFlags.BackgroundLoad)]
 
 	// Register the settings implementing class as providing support to UnrealVSPackage.
@@ -57,7 +58,7 @@ namespace UnrealVS
 	{
 		/** Constants */
 
-		private const string VersionString = "v1.55";
+		private const string VersionString = "v1.56";
 		private const string UnrealSolutionFileNamePrefix = "UE4";
 		private const string ExtensionName = "UnrealVS";
 		private const string CommandLineOptionKey = ExtensionName + "CommandLineMRU";
@@ -212,7 +213,7 @@ namespace UnrealVS
 			SolutionBuildManager.AdviseUpdateSolutionEvents(this, out UpdateSolutionEventsHandle);
 
 			// Create our command-line editor
-			CommandLineEditor = new CommandLineEditor();
+			CommandLineEditor.Initialize();
 
 			// Create our startup project selector
 			StartupProjectSelector = new StartupProjectSelector();
@@ -227,7 +228,7 @@ namespace UnrealVS
 			GenerateProjectFiles = new GenerateProjectFiles();
 
 			// Create Batch Builder tools
-			BatchBuilder = new BatchBuilder();
+			BatchBuilder.Initialize();
 
 			// Create the project menu quick builder
 			QuickBuilder = new QuickBuild();
@@ -486,6 +487,7 @@ namespace UnrealVS
 		/// <param name="stream">The stream to load the option data from.</param>
 		protected override void OnLoadOptions(string key, Stream stream)
 		{
+			Logging.WriteLine("Loading Options for key: " + key);
 			try
 			{
 				if (0 == string.Compare(key, CommandLineOptionKey))
@@ -847,7 +849,7 @@ namespace UnrealVS
 		UInt32 ProjectHierarchyEventsHandle;
 
 		/// Our command-line editing component
-		private CommandLineEditor CommandLineEditor;
+		private CommandLineEditor CommandLineEditor = new CommandLineEditor();
 
 		/// BuildStartupProject feature
 		private BuildStartupProject BuildStartupProject;
@@ -859,7 +861,7 @@ namespace UnrealVS
 		private GenerateProjectFiles GenerateProjectFiles;
 
 		/// Batch Builder button/command handler
-		private BatchBuilder BatchBuilder;
+		private BatchBuilder BatchBuilder = new BatchBuilder();
 
 		/// Ticker thread
 		private Thread Ticker;
