UnrealVS - Added more logging when running build batch jobs to catch crash

Potential fix for null dereference when starting batch builds without selecting a job first

diff --git a/Engine/Source/Programs/UnrealVS/BatchBuilderToolControl.xaml.cs b/Engine/Source/Programs/UnrealVS/BatchBuilderToolControl.xaml.cs
--- a/Engine/Source/Programs/UnrealVS/BatchBuilderToolControl.xaml.cs
+++ b/Engine/Source/Programs/UnrealVS/BatchBuilderToolControl.xaml.cs
@@ -833,8 +833,15 @@ namespace UnrealVS
 
 		private void RunBuildJobs()
 		{
+			if (BuildJobs == null)
+				throw new ArgumentNullException(nameof(BuildJobs));
+
+			if (_LastBuildJobsQueued == null)
+				throw new ArgumentNullException(nameof(_LastBuildJobsQueued));
+
 			int BuildJobCount = BuildJobs.Count;
 			SanityCheckBuildJobs();
+
 			if (BuildJobCount > BuildJobs.Count)
 			{
 				if (VSConstants.S_OK != VsShellUtilities.ShowMessageBox(ServiceProvider.GlobalProvider,
@@ -856,8 +863,11 @@ namespace UnrealVS
 				_BuildQueue.Enqueue(Job);
 			}
 
-			OutputPanelTitle = String.Format("{0} ({1})", OutputPanelTitlePrefix, _LastSelectedBuildJobSet.Name);
 			HasOutput = _BuildQueue.Count > 0;
+			if (_LastSelectedBuildJobSet != null)
+			{
+				OutputPanelTitle = String.Format("{0} ({1})", OutputPanelTitlePrefix, _LastSelectedBuildJobSet.Name);
+			}
 			OutputTab.IsSelected = HasOutput;
 
 			UpdateBusyState();
