Unreal VS 1.57

* Improved detection of UE solution by using a special section outputted to the solution, this replaces the need to add a special comment into the solution (or relying on naming convention as was the fallback). This still works with the old conventions to support being used in a UE4 solution.
* Inverted Quickbuild menu to dynamically populate the platforms instead - this allows UnrealVS to data drive the platform list. Unfortunatley due to how visual studio UI handles dynamic lists we can not have a menu that is dynamic and has children. Thus configurations are now hard coded instead - these are still only visible when they are actually supported, but it means UnrealVS needs to be updated when new configurations are added.
* Fixed issue setting debug arguments to net core projects (due to the change to launchSettings.json to store the debug arguments)

diff --git a/Engine/Source/Programs/UnrealBuildTool/ProjectFiles/VisualStudio/VCProjectFileGenerator.cs b/Engine/Source/Programs/UnrealBuildTool/ProjectFiles/VisualStudio/VCProjectFileGenerator.cs
--- a/Engine/Source/Programs/UnrealBuildTool/ProjectFiles/VisualStudio/VCProjectFileGenerator.cs
+++ b/Engine/Source/Programs/UnrealBuildTool/ProjectFiles/VisualStudio/VCProjectFileGenerator.cs
@@ -532,9 +532,6 @@ namespace UnrealBuildTool
 				throw new BuildException("Unexpected ProjectFileFormat");
 			}
 
-			// Embed custom comment used by UnrealVS to determine that this is a Unreal solution
-			VCSolutionFileContent.AppendLine("# UnrealEngineGeneratedSolutionVersion=" + AssemblyUtils.ExecutableVersion.ProductVersion);
-
 			IDictionary<MasterProjectFolder, Guid> ProjectFolderGuids = GenerateProjectFolderGuids(RootFolder);
 
 			// Solution folders, files and project entries
@@ -633,11 +630,11 @@ namespace UnrealBuildTool
 			{
 				VCSolutionFileContent.AppendLine("Global");
 				{
+					HashSet<UnrealTargetPlatform> PlatformsValidForProjects = new HashSet<UnrealTargetPlatform>();
 					{
 						VCSolutionFileContent.AppendLine("	GlobalSection(SolutionConfigurationPlatforms) = preSolution");
 
 						Dictionary<string, Tuple<UnrealTargetConfiguration, TargetType>> SolutionConfigurationsValidForProjects = new Dictionary<string, Tuple<UnrealTargetConfiguration, TargetType>>();
-						HashSet<UnrealTargetPlatform> PlatformsValidForProjects = new HashSet<UnrealTargetPlatform>();
 
 						foreach (UnrealTargetConfiguration CurConfiguration in SupportedConfigurations)
 						{
@@ -726,6 +723,14 @@ namespace UnrealBuildTool
 						VCSolutionFileContent.AppendLine("	EndGlobalSection");
 					}
 
+					// Embed UnrealVS section, which is parsed by that VSPackage (Extension) to know how to handle this solution
+					{
+						string UnrealVSGuid = "ddbf523f-7eb6-4887-bd51-85a714ff87eb";
+						VCSolutionFileContent.AppendLine("\t# UnrealVS Section");
+						VCSolutionFileContent.AppendLine("\tGlobalSection({0}) = preSolution", UnrealVSGuid);
+						VCSolutionFileContent.AppendLine("\t\tAvailablePlatforms={0}", string.Join(";", PlatformsValidForProjects.Select(platform => platform.ToString())));
+						VCSolutionFileContent.AppendLine("\tEndGlobalSection");
+					}
 
 					// Assign each project's "project configuration" to our "solution platform + configuration" pairs.  This
 					// also sets up which projects are actually built when building the solution.
diff --git a/Engine/Source/Programs/UnrealVS/CommandLineEditor.cs b/Engine/Source/Programs/UnrealVS/CommandLineEditor.cs
--- a/Engine/Source/Programs/UnrealVS/CommandLineEditor.cs
+++ b/Engine/Source/Programs/UnrealVS/CommandLineEditor.cs
@@ -5,6 +5,7 @@ using System.Diagnostics;
 using System.ComponentModel.Design;
 using System.IO;
 using System.Linq;
+using System.Text.Json;
 using Microsoft.VisualStudio;
 using Microsoft.VisualStudio.Shell;
 using Microsoft.VisualStudio.Shell.Interop;
@@ -12,6 +13,7 @@ using Microsoft.VisualStudio.OLE.Interop;
 using EnvDTE;
 using System.Runtime.InteropServices;
 using System.Collections.Generic;
+using System.Text.Json.Serialization;
 using EnvDTE80;
 using System.Timers;
 
@@ -19,6 +21,19 @@ namespace UnrealVS
 {
 	internal class CommandLineEditor
 	{
+		private class LaunchSettingsJson
+		{
+			[JsonPropertyName("profiles")]
+			public Dictionary<string, LaunchSettingsProfile> Profiles { get; set; }
+
+			public class LaunchSettingsProfile
+			{
+				[JsonPropertyName("commandName")]
+				public string CommandName { get; set; }
+				[JsonPropertyName("commandLineArgs")]
+				public string CommandLineArgs { get; set; }
+			}
+		}
 		/** constants */
 
 		private const int ComboID = 0x1030;
@@ -192,10 +207,38 @@ namespace UnrealVS
 									Text = "";
 								}
 							}
+							else
+							{
+								// net core projects move debugger arguments into a debug profile setup via launchSetings.json
+								string activeDebugProfile;
+								if (PropertyStorage.GetPropertyValue("ActiveDebugProfile", ConfigurationName, (uint) _PersistStorageType.PST_USER_FILE, out activeDebugProfile) != VSConstants.S_OK)
+								{
+									activeDebugProfile = null;
+								}
 
+								string launchSettingsPath = Path.Combine(Path.GetDirectoryName(SelectedStartupProject.FileName), "Properties", "launchSettings.json");
+								if (File.Exists(launchSettingsPath))
+								{
+									string jsonDoc = File.ReadAllText(launchSettingsPath);
+									LaunchSettingsJson settings = JsonSerializer.Deserialize<LaunchSettingsJson>(jsonDoc, new JsonSerializerOptions
+									{
+										PropertyNameCaseInsensitive = true
+									});
+
+									if (activeDebugProfile != null)
+									{
+										Text = settings.Profiles[activeDebugProfile].CommandLineArgs ?? "";
+									}
+									else
+									{
+										// if no active debug profile is set then VS uses the first one
+										Text = settings.Profiles.FirstOrDefault().Value?.CommandLineArgs ?? "";
+									}
+								}
+							}
 							// for "Game" projects automatically remove the game project filename from the start of the command line
 							var ActiveConfiguration = (SolutionConfiguration2)UnrealVSPackage.Instance.DTE.Solution.SolutionBuild.ActiveConfiguration;
-							if (UnrealVSPackage.Instance.IsUE4Loaded && Utils.IsGameProject(SelectedStartupProject) && Utils.HasUProjectCommandLineArg(ActiveConfiguration.Name))
+							if (UnrealVSPackage.Instance.IsUESolutionLoaded && Utils.IsGameProject(SelectedStartupProject) && Utils.HasUProjectCommandLineArg(ActiveConfiguration.Name))
 							{
 								string AutoPrefix = Utils.GetAutoUProjectCommandLinePrefix(SelectedStartupProject);
 								if (!string.IsNullOrEmpty(AutoPrefix))
@@ -291,7 +334,7 @@ namespace UnrealVS
 
 							// for "Game" projects automatically remove the game project filename from the start of the command line
 							var ActiveConfiguration = (SolutionConfiguration2)UnrealVSPackage.Instance.DTE.Solution.SolutionBuild.ActiveConfiguration;
-							if (UnrealVSPackage.Instance.IsUE4Loaded && Utils.IsGameProject(SelectedStartupProject) && Utils.HasUProjectCommandLineArg(ActiveConfiguration.Name))
+							if (UnrealVSPackage.Instance.IsUESolutionLoaded && Utils.IsGameProject(SelectedStartupProject) && Utils.HasUProjectCommandLineArg(ActiveConfiguration.Name))
 							{
 								string AutoPrefix = Utils.GetAutoUProjectCommandLinePrefix(SelectedStartupProject);
 								if (FullCommandLine.IndexOf(Utils.UProjectExtension, StringComparison.OrdinalIgnoreCase) < 0 &&
@@ -341,17 +384,63 @@ namespace UnrealVS
 							else
 							{
 								// For some reason, have to update C# projects this way, otherwise the properties page doesn't update. Conversely, SelectedConfiguration.Properties is always null for C++ projects in VS2017.
+								bool netCoreProject = false;
 								if (SelectedConfiguration.Properties != null)
 								{
 									foreach (Property Property in SelectedConfiguration.Properties)
 									{
 										if (Property.Name.Equals("StartArguments", StringComparison.InvariantCultureIgnoreCase))
 										{
-											Property.Value = FullCommandLine;
+											try
+											{
+												Property.Value = FullCommandLine;
+											}
+											catch (NotImplementedException)
+											{
+												// net core projects report containing start arguments command but does not support setting it
+												netCoreProject = true;
+											}
 											break;
 										}
 									}
 								}
+
+								if (netCoreProject)
+								{
+									// net core project use launchSettings.json to control debug arguments instead
+
+									string activeDebugProfile;
+									if (PropertyStorage.GetPropertyValue("ActiveDebugProfile", ProjectConfigurationName, (uint)_PersistStorageType.PST_USER_FILE, out activeDebugProfile) != VSConstants.S_OK)
+									{
+										activeDebugProfile = null;
+									}
+
+									string launchSettingsPath = Path.Combine(Path.GetDirectoryName(SelectedStartupProject.FileName), "Properties", "launchSettings.json");
+									if (File.Exists(launchSettingsPath))
+									{
+										string jsonDoc = File.ReadAllText(launchSettingsPath);
+										LaunchSettingsJson settings = JsonSerializer.Deserialize<LaunchSettingsJson>(jsonDoc, new JsonSerializerOptions
+										{
+											PropertyNameCaseInsensitive = true
+										});
+
+										if (activeDebugProfile != null)
+										{
+											settings.Profiles[activeDebugProfile].CommandLineArgs = FullCommandLine;
+										}
+										else
+										{
+											// if no active debug profile is set then VS uses the first one
+											settings.Profiles.First().Value.CommandLineArgs = FullCommandLine;
+										}
+
+										string json = JsonSerializer.Serialize(settings, new JsonSerializerOptions
+										{
+											WriteIndented = true,
+										});
+										File.WriteAllText(launchSettingsPath, json);
+									}
+								}
 							}
 						}
 					}
diff --git a/Engine/Source/Programs/UnrealVS/GenerateProjectFiles.cs b/Engine/Source/Programs/UnrealVS/GenerateProjectFiles.cs
--- a/Engine/Source/Programs/UnrealVS/GenerateProjectFiles.cs
+++ b/Engine/Source/Programs/UnrealVS/GenerateProjectFiles.cs
@@ -1,4 +1,4 @@
-﻿// Copyright Epic Games, Inc. All Rights Reserved.
+// Copyright Epic Games, Inc. All Rights Reserved.
 
 using System;
 using System.Diagnostics;
@@ -146,7 +146,7 @@ namespace UnrealVS
 		private string GetBatchFileName()
 		{
 			// Check to see if we have UE4.sln loaded
-			if (UnrealVSPackage.Instance.IsUE4Loaded)
+			if (UnrealVSPackage.Instance.IsUESolutionLoaded)
 			{ 
 				// We expect "GenerateProjectFiles.bat" to live in the same directory as the solution
 				return Path.Combine(Path.GetDirectoryName(UnrealVSPackage.Instance.SolutionFilepath), "GenerateProjectFiles.bat");
diff --git a/Engine/Source/Programs/UnrealVS/Properties/AssemblyInfo.cs b/Engine/Source/Programs/UnrealVS/Properties/AssemblyInfo.cs
--- a/Engine/Source/Programs/UnrealVS/Properties/AssemblyInfo.cs
+++ b/Engine/Source/Programs/UnrealVS/Properties/AssemblyInfo.cs
@@ -31,8 +31,8 @@ using System.Runtime.InteropServices;
 // You can specify all the values or you can default the Revision and Build Numbers 
 // by using the '*' as shown below:
 
-[assembly: AssemblyVersion("1.56.0.0")]
-[assembly: AssemblyFileVersion("1.56.0.0")]
+[assembly: AssemblyVersion("1.57.0.0")]
+[assembly: AssemblyFileVersion("1.57.0.0")]
 
 
 
diff --git a/Engine/Source/Programs/UnrealVS/ProvideSolutionProperties.cs b/Engine/Source/Programs/UnrealVS/ProvideSolutionProperties.cs
new file mode 100644
--- /dev/null
+++ b/Engine/Source/Programs/UnrealVS/ProvideSolutionProperties.cs
@@ -0,0 +1,48 @@
+using System;
+using System.Globalization;
+using Microsoft.VisualStudio.Shell;
+
+namespace UnrealVS
+{
+	/// <summary>
+	/// This attribute registers the package as a solution property parser
+	/// </summary>
+	[AttributeUsage(AttributeTargets.Class, AllowMultiple = true, Inherited = true)]
+	internal sealed class ProvideSolutionPropertiesAttribute : RegistrationAttribute
+	{
+		private string _propName;
+
+		public ProvideSolutionPropertiesAttribute(string propName)
+		{
+			_propName = propName;
+		}
+
+		public override void Register(RegistrationContext context)
+		{
+			context.Log.WriteLine(string.Format(CultureInfo.InvariantCulture, "ProvideSolutionProps: ({0} = {1})", context.ComponentType.GUID.ToString("B"), PropName));
+
+			Key childKey = null;
+
+			try
+			{
+				childKey = context.CreateKey(string.Format(CultureInfo.InvariantCulture, "{0}\\{1}", "SolutionPersistence", PropName));
+
+				childKey.SetValue(string.Empty, context.ComponentType.GUID.ToString("B").ToUpperInvariant());
+			}
+			finally
+			{
+				if (childKey != null) childKey.Close();
+			}
+		}
+
+		public override void Unregister(RegistrationContext context)
+		{
+			context.RemoveKey(string.Format(CultureInfo.InvariantCulture, "{0}\\{1}", "SolutionPersistence", PropName));
+		}
+
+		public string PropName
+		{
+			get { return _propName; }
+		}
+	}
+}
diff --git a/Engine/Source/Programs/UnrealVS/QuickBuild.cs b/Engine/Source/Programs/UnrealVS/QuickBuild.cs
--- a/Engine/Source/Programs/UnrealVS/QuickBuild.cs
+++ b/Engine/Source/Programs/UnrealVS/QuickBuild.cs
@@ -12,6 +12,7 @@ using System.Linq;
 using System.Runtime.InteropServices;
 using System.Text;
 using System.Threading.Tasks;
+using System.Windows.Forms.VisualStyles;
 
 namespace UnrealVS
 {
@@ -61,24 +62,24 @@ namespace UnrealVS
 	}
 
 	/// <summary>
-	/// Class represents a single platform, quick build menu. Wraps QuickBuildMenuCommand.
+	/// Class represents a single configuration quick build menu. Wraps QuickBuildMenuCommand.
 	/// </summary>
-	class PlaformMenuContents
+	class ConfigurationMenuContents
 	{
 		/// The command id of the placeholder command that sets the start of the command range for this menu
 		private readonly int DynamicStartCommandId;
-		/// The name of the platform. Matches the platform string in the solution configs exactly.
-		private readonly string PlaformName;
+		/// The name of the configuration. Matches the config string in the solution configs exactly.
+		private readonly string ConfigurationName;
 
 		/// <summary>
 		/// Construct a platform menu for a given, named platform, starting at the given command id
 		/// </summary>
 		/// <param name="Name">The name of the platform</param>
 		/// <param name="InDynamicStartCommandId">The start of the command range for this menu</param>
-		public PlaformMenuContents(string Name, int InDynamicStartCommandId)
+		public ConfigurationMenuContents(string Name, int InDynamicStartCommandId)
 		{
 			DynamicStartCommandId = InDynamicStartCommandId;
-			PlaformName = Name;
+			ConfigurationName = Name;
 
 			// Create the dynamic menu
 			var MenuOleCommand = new QuickBuildMenuCommand(
@@ -86,7 +87,7 @@ namespace UnrealVS
 				IsValidDynamicItem,
 				OnInvokedDynamicItem,
 				OnBeforeQueryStatusDynamicItem
-				);
+			);
 			UnrealVSPackage.Instance.MenuCommandService.AddCommand(MenuOleCommand);
 		}
 
@@ -98,7 +99,7 @@ namespace UnrealVS
 		/// <returns>True, if the command is valid for this menu</returns>
 		private bool IsValidDynamicItem(int cmdId)
 		{
-			int ConfigCount = QuickBuild.CachedSolutionConfigNames.Length;
+			int ConfigCount = QuickBuild.CachedSolutionConfigPlatforms.Length;
 
 			// The match is valid if the command ID is >= the id of our root dynamic start item
 			// and the command ID minus the ID of our root dynamic start item
@@ -127,7 +128,7 @@ namespace UnrealVS
 			if (SelectedProject == null) return;
 
 			// Builds the selected project with the clicked platform and config
-			Utils.ExecuteProjectBuild(SelectedProject, MenuCommand.Text, PlaformName, BatchBuilderToolControl.BuildJob.BuildJobType.Build, null, null);
+			Utils.ExecuteProjectBuild(SelectedProject, ConfigurationName, MenuCommand.Text, BatchBuilderToolControl.BuildJob.BuildJobType.Build, null, null);
 		}
 
 		/// <summary>
@@ -144,10 +145,10 @@ namespace UnrealVS
 			int CommandId = isRootItem ? DynamicStartCommandId : MenuCommand.MatchedCommandId;
 			int DynCommandIdx = CommandId - DynamicStartCommandId;
 
-			// Set the text label based on the index and the solution configs array cached in QuickBuild
-			if (DynCommandIdx < QuickBuild.CachedSolutionConfigNames.Length)
+			// Set the text label based on the index and the solution platforms array cached in QuickBuild
+			if (DynCommandIdx < QuickBuild.CachedSolutionConfigPlatforms.Length)
 			{
-				MenuCommand.Text = QuickBuild.CachedSolutionConfigNames[DynCommandIdx];
+				MenuCommand.Text = QuickBuild.CachedSolutionConfigPlatforms[DynCommandIdx];
 			}
 
 			// Clear the id now that the query is done
@@ -182,22 +183,32 @@ namespace UnrealVS
 
 		/// Hide/shows the while Quick Build menu tree
 		private static bool bIsActive = false;
-
+		
 		/// List of submenus and their details - must match the values in the vsct file
 		private readonly SubMenu[] SubMenus = new[]
-			{
-				new SubMenu {Name = "Win64", SubMenuId = 0x1500, DynamicStartCommandId = 0x1530},
-				new SubMenu {Name = "Win32", SubMenuId = 0x1600, DynamicStartCommandId = 0x1630},
-				new SubMenu {Name = "Mac", SubMenuId = 0x1700, DynamicStartCommandId = 0x1730},
-				new SubMenu {Name = "XboxOne", SubMenuId = 0x1800, DynamicStartCommandId = 0x1830},
-				new SubMenu {Name = "PS4", SubMenuId = 0x1900, DynamicStartCommandId = 0x1930},
-				new SubMenu {Name = "IOS", SubMenuId = 0x1A00, DynamicStartCommandId = 0x1A30},
-				new SubMenu {Name = "Android", SubMenuId = 0x1B00, DynamicStartCommandId = 0x1B30},
-				new SubMenu {Name = "WinRT", SubMenuId = 0x1C00, DynamicStartCommandId = 0x1C30},
-				new SubMenu {Name = "WinRT_ARM", SubMenuId = 0x1D00, DynamicStartCommandId = 0x1D30},
-				new SubMenu {Name = "Linux", SubMenuId = 0x1F00, DynamicStartCommandId = 0x1F30},
-				new SubMenu {Name = "Switch", SubMenuId = 0x2000, DynamicStartCommandId = 0x2030},
-			};
+		{
+			new SubMenu {Name = "Debug", SubMenuId = 0x1500, DynamicStartCommandId = 0x1530},
+			new SubMenu {Name = "Debug Client", SubMenuId = 0x1600, DynamicStartCommandId = 0x1630},
+			new SubMenu {Name = "Debug Editor", SubMenuId = 0x1700, DynamicStartCommandId = 0x1730},
+			new SubMenu {Name = "Debug Server", SubMenuId = 0x1800, DynamicStartCommandId = 0x1830},
+			new SubMenu {Name = "DebugGame", SubMenuId = 0x1900, DynamicStartCommandId = 0x1930},
+			new SubMenu {Name = "DebugGame Client", SubMenuId = 0x1A00, DynamicStartCommandId = 0x1A30},
+			new SubMenu {Name = "DebugGame Editor", SubMenuId = 0x1B00, DynamicStartCommandId = 0x1B30},
+			new SubMenu {Name = "DebugGame Server", SubMenuId = 0x1C00, DynamicStartCommandId = 0x1C30},
+			new SubMenu {Name = "Development", SubMenuId = 0x1D00, DynamicStartCommandId = 0x1D30},
+			new SubMenu {Name = "Development Client", SubMenuId = 0x1E00, DynamicStartCommandId = 0x1E30},
+			new SubMenu {Name = "Development Editor", SubMenuId = 0x1F00, DynamicStartCommandId = 0x1F30},
+			new SubMenu {Name = "Development Server", SubMenuId = 0x2000, DynamicStartCommandId = 0x2030},
+			new SubMenu {Name = "Shipping", SubMenuId = 0x2100, DynamicStartCommandId = 0x2130},
+			new SubMenu {Name = "Shipping Client", SubMenuId = 0x2200, DynamicStartCommandId = 0x2230},
+			new SubMenu {Name = "Shipping Editor", SubMenuId = 0x2300, DynamicStartCommandId = 0x2330},
+			new SubMenu {Name = "Shipping Server", SubMenuId = 0x2400, DynamicStartCommandId = 0x2430},
+			new SubMenu {Name = "Test", SubMenuId = 0x2500, DynamicStartCommandId = 0x2530},
+			new SubMenu {Name = "Test Client", SubMenuId = 0x2600, DynamicStartCommandId = 0x2630},
+			new SubMenu {Name = "Test Editor", SubMenuId = 0x2700, DynamicStartCommandId = 0x2730},
+			new SubMenu {Name = "Test Server", SubMenuId = 0x2800, DynamicStartCommandId = 0x2830},
+
+		};
 
 		/// The main root command of the Quick Build menu hierarchy - used to hide it when not active
 		private readonly OleMenuCommand QuickBuildCommand;
@@ -206,19 +217,20 @@ namespace UnrealVS
 		/// These represent the items that can be added to the menu that lists the platforms.
 		/// Each one is a submenu containing items for each config.
 		/// </summary>
-		private readonly Dictionary<string, OleMenuCommand> AllPlaformMenus = new Dictionary<string, OleMenuCommand>();
+		private readonly Dictionary<string, OleMenuCommand> AllConfigurationMenus = new Dictionary<string, OleMenuCommand>();
 
 		/// <summary>
-		/// These represent the items shown in the menu that lists the platforms.
-		/// It is a subset of AllPlaformMenus with only the loaded platforms.
+		/// These represent the items shown in the menu that lists the configurations.
+		/// It is a subset of AllConfigurationMenus with only the loaded platforms.
 		/// Each one is a submenu containing items for each config.
 		/// </summary>
-		private readonly Dictionary<string, OleMenuCommand> ActivePlaformMenus = new Dictionary<string, OleMenuCommand>();
+		private readonly Dictionary<string, OleMenuCommand> ActiveConfigurationMenus = new Dictionary<string, OleMenuCommand>();
 
 		/// <summary>
-		/// These represent the items in plaform-specific menus.
+		/// These holds the objects that generate the per platform submenus for each configuration.
 		/// </summary>
-		private readonly Dictionary<string, PlaformMenuContents> PlaformMenusContents = new Dictionary<string, PlaformMenuContents>();
+		// ReSharper disable once CollectionNeverQueried.Local - This is only invoked by visual studio to populate the submenus.
+		private readonly Dictionary<string, ConfigurationMenuContents> ConfigurationMenusContents = new Dictionary<string, ConfigurationMenuContents>();
 
 		/// VSConstants.UICONTEXT_SolutionBuilding translated into a cookie used to access UI ctxt state
 		private readonly uint SolutionBuildingUIContextCookie;
@@ -226,7 +238,6 @@ namespace UnrealVS
 		/** properties */
 
 		public static bool IsActive { get { return bIsActive; } }
-		public static string[] CachedSolutionConfigNames { get { return SolutionConfigNames; } }
 		public static string[] CachedSolutionConfigPlatforms { get { return SolutionConfigPlatforms; } }
 
 		/** methods */
@@ -237,17 +248,16 @@ namespace UnrealVS
 			QuickBuildCommand = new OleMenuCommand(null, null, OnQuickBuildQuery, new CommandID(GuidList.UnrealVSCmdSet, ProjectQuickBuildMenuID));
 			QuickBuildCommand.BeforeQueryStatus += OnQuickBuildQuery;
 			UnrealVSPackage.Instance.MenuCommandService.AddCommand(QuickBuildCommand);
-
-			// platform sub-menus
+			
+			// configuration sub-menus
 			foreach (var SubMenu in SubMenus)
 			{
 				var SubMenuCommand = new OleMenuCommand(null, new CommandID(GuidList.UnrealVSCmdSet, SubMenu.SubMenuId));
 				SubMenuCommand.BeforeQueryStatus += OnQuickBuildSubMenuQuery;
 				UnrealVSPackage.Instance.MenuCommandService.AddCommand(SubMenuCommand);
-				AllPlaformMenus.Add(SubMenu.Name, SubMenuCommand);
-				PlaformMenusContents.Add(SubMenu.Name, new PlaformMenuContents(SubMenu.Name, SubMenu.DynamicStartCommandId));
+				AllConfigurationMenus.Add(SubMenu.Name, SubMenuCommand);
+				ConfigurationMenusContents.Add(SubMenu.Name, new ConfigurationMenuContents(SubMenu.Name, SubMenu.DynamicStartCommandId));
 			}
-
 			// cache the cookie for UICONTEXT_SolutionBuilding
 			UnrealVSPackage.Instance.SelectionManager.GetCmdUIContextCookie(VSConstants.UICONTEXT_SolutionBuilding, out SolutionBuildingUIContextCookie);
 
@@ -285,25 +295,25 @@ namespace UnrealVS
 		private void OnQuickBuildSubMenuQuery(object sender, EventArgs e)
 		{
 			var SubMenuCommand = (OleMenuCommand)sender;
-			SubMenuCommand.Visible = ActivePlaformMenus.ContainsValue(SubMenuCommand);
+			SubMenuCommand.Visible = ActiveConfigurationMenus.ContainsValue(SubMenuCommand);
 			SubMenuCommand.Enabled = bIsActive;
 		}
 
 		/// <summary>
 		/// Called when a solution loads. Caches the solution build configs and sets the platform menus' visibility.
-		/// Only platforms found in the laoded solution's list are shown.
+		/// Only configs found in the loaded solution's list are shown.
 		/// </summary>
 		private void OnSolutionChanged()
 		{
 			CacheBuildConfigs();
 
-			ActivePlaformMenus.Clear();
-			Logging.WriteLine("Updating solution menu state, currently active config platforms are: " + string.Join(", ", SolutionConfigPlatforms));
+			ActiveConfigurationMenus.Clear();
+			Logging.WriteLine("Updating solution menu state, currently active configs are: " + string.Join(", ", SolutionConfigNames));
 			foreach (var SubMenu in SubMenus)
 			{
-				if (SolutionConfigPlatforms.Any(Platform => string.Compare(Platform, SubMenu.Name, StringComparison.InvariantCultureIgnoreCase) == 0))
+				if (SolutionConfigNames.Any(Config => string.Compare(Config, SubMenu.Name, StringComparison.InvariantCultureIgnoreCase) == 0))
 				{
-					ActivePlaformMenus.Add(SubMenu.Name, AllPlaformMenus[SubMenu.Name]);
+					ActiveConfigurationMenus.Add(SubMenu.Name, AllConfigurationMenus[SubMenu.Name]);
 				}
 			}
 		}
@@ -313,14 +323,11 @@ namespace UnrealVS
 		/// </summary>
 		private void CacheBuildConfigs()
 		{
-			SolutionConfigurations SolutionConfigs =
-				UnrealVSPackage.Instance.DTE.Solution.SolutionBuild.SolutionConfigurations;
+			SolutionConfigurations SolutionConfigs = UnrealVSPackage.Instance.DTE.Solution.SolutionBuild.SolutionConfigurations;
 
-			SolutionConfigPlatforms =
-				(from SolutionConfiguration2 Sc in SolutionConfigs select Sc.PlatformName).Distinct().ToArray();
+			SolutionConfigPlatforms = (from SolutionConfiguration2 Sc in SolutionConfigs select Sc.PlatformName).Distinct().ToArray();
 
-			SolutionConfigNames =
-							(from SolutionConfiguration2 Sc in SolutionConfigs select Sc.Name).Distinct().ToArray();
+			SolutionConfigNames = (from SolutionConfiguration2 Sc in SolutionConfigs select Sc.Name).Distinct().ToArray();
 
 		}
 
diff --git a/Engine/Source/Programs/UnrealVS/UnrealVSPackage.cs b/Engine/Source/Programs/UnrealVS/UnrealVSPackage.cs
--- a/Engine/Source/Programs/UnrealVS/UnrealVSPackage.cs
+++ b/Engine/Source/Programs/UnrealVS/UnrealVSPackage.cs
@@ -8,16 +8,29 @@ using System.Linq;
 using System.Runtime.InteropServices;
 using System.ComponentModel.Design;
 using System.Threading;
+using System.Windows;
+using System.Windows.Forms;
 using Microsoft.VisualStudio;
 using Microsoft.VisualStudio.Shell.Interop;
 using Microsoft.VisualStudio.OLE.Interop;
 using Microsoft.VisualStudio.Shell;
 using EnvDTE;
 using EnvDTE80;
+using MessageBox = System.Windows.Forms.MessageBox;
 using Thread = System.Threading.Thread;
 
 namespace UnrealVS
 {
+	public class UnrealSolutionProperties
+	{
+		public UnrealSolutionProperties(string[] InAvailablePlatforms)
+		{
+			AvailablePlatforms = InAvailablePlatforms;
+		}
+
+		public string[] AvailablePlatforms { get; }
+	}
+
 	// This attribute tells the PkgDef creation utility (CreatePkgDef.exe) that this class is a VS package.
 	[PackageRegistration( UseManagedResourcesOnly = true, AllowsBackgroundLoading = true )]
 
@@ -43,7 +56,7 @@ namespace UnrealVS
 
 	// This attribute registers an options page for the package.
 	[ProvideOptionPage(typeof(UnrealVsOptions), ExtensionName, "General", 101, 102, true)]
-
+	[ProvideSolutionProperties(GuidList.UnrealVSPackageString)]
 	/// <summary>
 	/// UnrealVSPackage implements Package abstract class.  This is the main class that is registered
 	/// with Visual Studio shell and serves as the entry point into our extension
@@ -53,12 +66,13 @@ namespace UnrealVS
 		IVsSolutionEvents,		// This interface allows us to register to be notified of events such as opening a project
 		IVsUpdateSolutionEvents,// Allows us to register to be notified of events such as active config changes
 		IVsSelectionEvents,		// Allows us to be notified when the startup project has changed to a different project
-		IVsHierarchyEvents,		// Allows us to be notified when a hierarchy (the startup project) has had properties changed
+		IVsHierarchyEvents,     // Allows us to be notified when a hierarchy (the startup project) has had properties changed
+		IVsPersistSolutionProps, // Allows us to read props added to the solution to determine if the solution is a unreal solution
 		IDisposable
 	{
 		/** Constants */
 
-		private const string VersionString = "v1.56";
+		private const string VersionString = "v1.57";
 		private const string UnrealSolutionFileNamePrefix = "UE4";
 		private const string ExtensionName = "UnrealVS";
 		private const string CommandLineOptionKey = ExtensionName + "CommandLineMRU";
@@ -151,12 +165,18 @@ namespace UnrealVS
 		public IVsMonitorSelection SelectionManager { get; private set; }
 
 		/// Variable keeps track of whether a supported Unreal solution is loaded
-		public bool IsUE4Loaded { get; private set; }
+		public bool IsUESolutionLoaded
+		{
+			get { return _IsUESolutionLoaded.GetValueOrDefault(false); }
+			private set { _IsUESolutionLoaded = value; }
+		}
 
 		/// Variable keeps track of the loaded solution
 		private string _SolutionFilepath;
 		public string SolutionFilepath { get { return _SolutionFilepath; } }
 
+		public UnrealSolutionProperties UnrealSolutionProperties { get; private set; }
+
 		/** Methods */
 
 		/// <summary>
@@ -745,6 +765,73 @@ namespace UnrealVS
 			return VSConstants.S_OK;
 		}
 
+
+#region IVsPersistSolutionProps
+		int IVsPersistSolutionProps.SaveUserOptions(IVsSolutionPersistence pPersistence)
+		{
+			return VSConstants.S_OK;
+		}
+
+		int IVsPersistSolutionProps.LoadUserOptions(IVsSolutionPersistence pPersistence, uint grfLoadOpts)
+		{
+			return VSConstants.S_OK;
+		}
+
+		int IVsPersistSolutionProps.WriteUserOptions(IStream pOptionsStream, string pszKey)
+		{
+			return VSConstants.S_OK;
+		}
+
+		int IVsPersistSolutionProps.ReadUserOptions(IStream pOptionsStream, string pszKey)
+		{
+			return VSConstants.S_OK;
+		}
+
+		int IVsPersistSolutionProps.QuerySaveSolutionProps(IVsHierarchy pHierarchy, VSQUERYSAVESLNPROPS[] pqsspSave)
+		{
+			return VSConstants.S_OK;
+		}
+
+		int IVsPersistSolutionProps.SaveSolutionProps(IVsHierarchy pHierarchy, IVsSolutionPersistence pPersistence)
+		{
+			return VSConstants.S_OK;
+		}
+
+		int IVsPersistSolutionProps.WriteSolutionProps(IVsHierarchy pHierarchy, string pszKey, IPropertyBag pPropBag)
+		{
+			return VSConstants.S_OK;
+		}
+
+		int IVsPersistSolutionProps.ReadSolutionProps(IVsHierarchy pHierarchy, string pszProjectName, string pszProjectMk, string pszKey, int fPreLoad, IPropertyBag pPropBag)
+		{
+			if (!string.Equals(pszKey, GuidList.UnrealVSPackageString, StringComparison.InvariantCultureIgnoreCase))
+				return VSConstants.S_OK;
+
+			object availablePlatformsObject;
+			pPropBag.Read("AvailablePlatforms", out availablePlatformsObject, null, (uint)VarEnum.VT_BSTR, pPropBag);
+
+			string[] availablePlatforms = null;
+			if (availablePlatformsObject != null)
+			{
+				string availablePlatformsString = availablePlatformsObject as string;
+				availablePlatforms = availablePlatformsString.Split(';');
+			}
+
+			UnrealSolutionProperties = new UnrealSolutionProperties(availablePlatforms);
+
+			// a UnrealVS section was found so we consider this solution a unreal solution
+			IsUESolutionLoaded = true;
+
+			return VSConstants.S_OK;
+		}
+
+		int IVsPersistSolutionProps.OnProjectLoadFailure(IVsHierarchy pStubHierarchy, string pszProjectName, string pszProjectMk, string pszKey)
+		{
+			return VSConstants.S_OK;
+		}
+
+#endregion
+
 		// IVsUpdateSolutionEvents Interface
 
 		public int UpdateSolution_Begin(ref int pfCancelUpdate)
@@ -798,13 +885,18 @@ namespace UnrealVS
 		{
 			if (!DTE.Solution.IsOpen)
 			{
-				IsUE4Loaded = false;
+				IsUESolutionLoaded = false;
 				return;
 			}
 
 			string SolutionDirectory, UserOptsFile;
 			SolutionManager.GetSolutionInfo(out SolutionDirectory, out _SolutionFilepath, out UserOptsFile);
 
+			// if ReadProps found a valud UE solution we do not need to check these legacy definitions of a solution
+			if (_IsUESolutionLoaded.HasValue)
+				return;
+
+			// Legacy paths for determing if a solution is a unreal solution by checking for tags & solution name.
 			var SolutionLines = new string[0];
 			try
 			{
@@ -819,12 +911,12 @@ namespace UnrealVS
 			if (UBTLine != null)
 			{
 				_UBTVersion = UBTLine.Trim().Substring(UBTTag.Length);
-				IsUE4Loaded = true;
+				IsUESolutionLoaded = true;
 			}
 			else
 			{
 				_UBTVersion = string.Empty;
-				IsUE4Loaded =
+				IsUESolutionLoaded =
 					(
 						_SolutionFilepath != null &&
 						Path.GetFileName(_SolutionFilepath).StartsWith(UnrealSolutionFileNamePrefix, StringComparison.OrdinalIgnoreCase)
@@ -871,7 +963,8 @@ namespace UnrealVS
 
 		private string _UBTVersion = string.Empty;
 
-		private readonly List<string> LoadedProjectPaths = new List<string>(); 
+		private readonly List<string> LoadedProjectPaths = new List<string>();
+		private bool? _IsUESolutionLoaded;
 
 		/// Obtains the DTE2 interface for this instance of VS from the RunningObjectTable
 		private static DTE2 GetDTE2ForCurrentInstance(DTE DTE)
@@ -900,6 +993,7 @@ namespace UnrealVS
 			}
 			return null;
 		}
+
 	}
 
 	internal static class NativeMethods
diff --git a/Engine/Source/Programs/UnrealVS/UnrealVS_2019.csproj b/Engine/Source/Programs/UnrealVS/UnrealVS_2019.csproj
--- a/Engine/Source/Programs/UnrealVS/UnrealVS_2019.csproj
+++ b/Engine/Source/Programs/UnrealVS/UnrealVS_2019.csproj
@@ -44,6 +44,8 @@
     <ApplicationVersion>1.0.0.%2a</ApplicationVersion>
     <UseApplicationTrust>false</UseApplicationTrust>
     <BootstrapperEnabled>true</BootstrapperEnabled>
+    <NuGetPackageImportStamp>
+    </NuGetPackageImportStamp>
   </PropertyGroup>
   <PropertyGroup Condition=" '$(Configuration)|$(Platform)' == 'VS12d|AnyCPU' ">
     <DebugSymbols>true</DebugSymbols>
@@ -81,6 +83,9 @@
     <Reference Include="envdte80, Version=8.0.0.0, Culture=neutral, PublicKeyToken=b03f5f7f11d50a3a">
       <EmbedInteropTypes>True</EmbedInteropTypes>
     </Reference>
+    <Reference Include="Microsoft.Bcl.AsyncInterfaces, Version=1.0.0.0, Culture=neutral, PublicKeyToken=cc7b13ffcd2ddd51, processorArchitecture=MSIL">
+      <HintPath>packages\Microsoft.Bcl.AsyncInterfaces.1.1.0\lib\net461\Microsoft.Bcl.AsyncInterfaces.dll</HintPath>
+    </Reference>
     <Reference Include="Microsoft.CSharp" />
     <Reference Include="Microsoft.VisualStudio.OLE.Interop, Version=7.1.40304.0, Culture=neutral, PublicKeyToken=b03f5f7f11d50a3a" />
     <Reference Include="Microsoft.VisualStudio.Shell.15.0, Version=16.0.0.0, Culture=neutral, PublicKeyToken=b03f5f7f11d50a3a, processorArchitecture=MSIL" />
@@ -102,20 +107,59 @@
     <Reference Include="Microsoft.VisualStudio.Shell.Interop.8.0, Version=8.0.0.0, Culture=neutral, PublicKeyToken=b03f5f7f11d50a3a" />
     <Reference Include="Microsoft.VisualStudio.Shell.Interop.9.0, Version=9.0.0.0, Culture=neutral, PublicKeyToken=b03f5f7f11d50a3a" />
     <Reference Include="Microsoft.VisualStudio.TextManager.Interop, Version=7.1.40304.0, Culture=neutral, PublicKeyToken=b03f5f7f11d50a3a" />
-    <Reference Include="Microsoft.VisualStudio.Threading, Version=16.0.0.0, Culture=neutral, PublicKeyToken=b03f5f7f11d50a3a" />
+    <Reference Include="Microsoft.VisualStudio.Threading, Version=16.7.0.0, Culture=neutral, PublicKeyToken=b03f5f7f11d50a3a, processorArchitecture=MSIL">
+      <HintPath>packages\Microsoft.VisualStudio.Threading.16.7.56\lib\net472\Microsoft.VisualStudio.Threading.dll</HintPath>
+    </Reference>
+    <Reference Include="Microsoft.VisualStudio.Validation, Version=15.3.0.0, Culture=neutral, PublicKeyToken=b03f5f7f11d50a3a, processorArchitecture=MSIL">
+      <HintPath>packages\Microsoft.VisualStudio.Validation.15.5.31\lib\netstandard2.0\Microsoft.VisualStudio.Validation.dll</HintPath>
+    </Reference>
     <Reference Include="Microsoft.VisualStudio.VCProject, Version=16.0.0.0, Culture=neutral, PublicKeyToken=b03f5f7f11d50a3a">
       <EmbedInteropTypes>True</EmbedInteropTypes>
     </Reference>
     <Reference Include="Microsoft.VisualStudio.VCProjectEngine, Version=16.0.0.0, Culture=neutral, PublicKeyToken=b03f5f7f11d50a3a">
       <EmbedInteropTypes>True</EmbedInteropTypes>
     </Reference>
+    <Reference Include="Microsoft.Win32.Registry, Version=4.1.2.0, Culture=neutral, PublicKeyToken=b03f5f7f11d50a3a, processorArchitecture=MSIL">
+      <HintPath>packages\Microsoft.Win32.Registry.4.6.0\lib\net461\Microsoft.Win32.Registry.dll</HintPath>
+    </Reference>
     <Reference Include="PresentationCore" />
     <Reference Include="PresentationFramework" />
     <Reference Include="System" />
+    <Reference Include="System.Buffers, Version=4.0.3.0, Culture=neutral, PublicKeyToken=cc7b13ffcd2ddd51, processorArchitecture=MSIL">
+      <HintPath>packages\System.Buffers.4.5.1\lib\net461\System.Buffers.dll</HintPath>
+    </Reference>
     <Reference Include="System.Core" />
     <Reference Include="System.Data" />
     <Reference Include="System.Design" />
     <Reference Include="System.Drawing" />
+    <Reference Include="System.Memory, Version=4.0.1.1, Culture=neutral, PublicKeyToken=cc7b13ffcd2ddd51, processorArchitecture=MSIL">
+      <HintPath>packages\System.Memory.4.5.4\lib\net461\System.Memory.dll</HintPath>
+    </Reference>
+    <Reference Include="System.Numerics" />
+    <Reference Include="System.Numerics.Vectors, Version=4.1.4.0, Culture=neutral, PublicKeyToken=b03f5f7f11d50a3a, processorArchitecture=MSIL">
+      <HintPath>packages\System.Numerics.Vectors.4.5.0\lib\net46\System.Numerics.Vectors.dll</HintPath>
+    </Reference>
+    <Reference Include="System.Runtime.CompilerServices.Unsafe, Version=4.0.6.0, Culture=neutral, PublicKeyToken=b03f5f7f11d50a3a, processorArchitecture=MSIL">
+      <HintPath>packages\System.Runtime.CompilerServices.Unsafe.4.7.1\lib\net461\System.Runtime.CompilerServices.Unsafe.dll</HintPath>
+    </Reference>
+    <Reference Include="System.Security.AccessControl, Version=4.1.2.0, Culture=neutral, PublicKeyToken=b03f5f7f11d50a3a, processorArchitecture=MSIL">
+      <HintPath>packages\System.Security.AccessControl.4.6.0\lib\net461\System.Security.AccessControl.dll</HintPath>
+    </Reference>
+    <Reference Include="System.Security.Principal.Windows, Version=4.1.2.0, Culture=neutral, PublicKeyToken=b03f5f7f11d50a3a, processorArchitecture=MSIL">
+      <HintPath>packages\System.Security.Principal.Windows.4.6.0\lib\net461\System.Security.Principal.Windows.dll</HintPath>
+    </Reference>
+    <Reference Include="System.Text.Encodings.Web, Version=4.0.5.0, Culture=neutral, PublicKeyToken=cc7b13ffcd2ddd51, processorArchitecture=MSIL">
+      <HintPath>packages\System.Text.Encodings.Web.4.7.1\lib\net461\System.Text.Encodings.Web.dll</HintPath>
+    </Reference>
+    <Reference Include="System.Text.Json, Version=4.0.1.2, Culture=neutral, PublicKeyToken=cc7b13ffcd2ddd51, processorArchitecture=MSIL">
+      <HintPath>packages\System.Text.Json.4.7.2\lib\net461\System.Text.Json.dll</HintPath>
+    </Reference>
+    <Reference Include="System.Threading.Tasks.Extensions, Version=4.2.0.1, Culture=neutral, PublicKeyToken=cc7b13ffcd2ddd51, processorArchitecture=MSIL">
+      <HintPath>packages\System.Threading.Tasks.Extensions.4.5.4\lib\net461\System.Threading.Tasks.Extensions.dll</HintPath>
+    </Reference>
+    <Reference Include="System.ValueTuple, Version=4.0.3.0, Culture=neutral, PublicKeyToken=cc7b13ffcd2ddd51, processorArchitecture=MSIL">
+      <HintPath>packages\System.ValueTuple.4.5.0\lib\net47\System.ValueTuple.dll</HintPath>
+    </Reference>
     <Reference Include="System.Windows.Forms" />
     <Reference Include="System.Xaml" />
     <Reference Include="System.Xml" />
@@ -156,6 +200,7 @@
     <Compile Include="ProfileManager.cs">
       <SubType>Component</SubType>
     </Compile>
+    <Compile Include="ProvideSolutionProperties.cs" />
     <Compile Include="QuickBuild.cs" />
     <Compile Include="Resources.Designer.cs">
       <AutoGen>True</AutoGen>
@@ -184,6 +229,7 @@
     </EmbeddedResource>
   </ItemGroup>
   <ItemGroup>
+    <None Include="app.config" />
     <None Include="packages.config" />
   </ItemGroup>
   <ItemGroup>
@@ -244,6 +290,14 @@
       <SubType>Designer</SubType>
     </Page>
   </ItemGroup>
+  <ItemGroup>
+    <Analyzer Include="packages\Microsoft.VisualStudio.Threading.Analyzers.16.7.56\analyzers\cs\Microsoft.VisualStudio.Threading.Analyzers.CodeFixes.dll" />
+    <Analyzer Include="packages\Microsoft.VisualStudio.Threading.Analyzers.16.7.56\analyzers\cs\Microsoft.VisualStudio.Threading.Analyzers.CSharp.dll" />
+    <Analyzer Include="packages\Microsoft.VisualStudio.Threading.Analyzers.16.7.56\analyzers\cs\Microsoft.VisualStudio.Threading.Analyzers.dll" />
+    <Analyzer Include="packages\Microsoft.VisualStudio.Threading.Analyzers.16.7.56\analyzers\vb\Microsoft.VisualStudio.Threading.Analyzers.CodeFixes.dll" />
+    <Analyzer Include="packages\Microsoft.VisualStudio.Threading.Analyzers.16.7.56\analyzers\vb\Microsoft.VisualStudio.Threading.Analyzers.dll" />
+    <Analyzer Include="packages\Microsoft.VisualStudio.Threading.Analyzers.16.7.56\analyzers\vb\Microsoft.VisualStudio.Threading.Analyzers.VisualBasic.dll" />
+  </ItemGroup>
   <PropertyGroup>
     <UseCodebase>true</UseCodebase>
   </PropertyGroup>
@@ -344,6 +398,13 @@
   <PropertyGroup>
     <PreBuildEvent>copy "$(ProjectDir)TargetManifests\$(ConfigurationName).source.extension.vsixmanifest" "$(ProjectDir)Generated\source.extension.vsixmanifest" /Y</PreBuildEvent>
   </PropertyGroup>
+  <Import Project="packages\Microsoft.VisualStudio.Threading.Analyzers.16.7.56\build\Microsoft.VisualStudio.Threading.Analyzers.targets" Condition="Exists('packages\Microsoft.VisualStudio.Threading.Analyzers.16.7.56\build\Microsoft.VisualStudio.Threading.Analyzers.targets')" />
+  <Target Name="EnsureNuGetPackageBuildImports" BeforeTargets="PrepareForBuild">
+    <PropertyGroup>
+      <ErrorText>This project references NuGet package(s) that are missing on this computer. Use NuGet Package Restore to download them.  For more information, see http://go.microsoft.com/fwlink/?LinkID=322105. The missing file is {0}.</ErrorText>
+    </PropertyGroup>
+    <Error Condition="!Exists('packages\Microsoft.VisualStudio.Threading.Analyzers.16.7.56\build\Microsoft.VisualStudio.Threading.Analyzers.targets')" Text="$([System.String]::Format('$(ErrorText)', 'packages\Microsoft.VisualStudio.Threading.Analyzers.16.7.56\build\Microsoft.VisualStudio.Threading.Analyzers.targets'))" />
+  </Target>
   <!-- To modify your build process, add your task inside one of the targets below and uncomment it. 
        Other similar extension points exist, see Microsoft.Common.targets.
   <Target Name="BeforeBuild">
diff --git a/Engine/Source/Programs/UnrealVS/app.config b/Engine/Source/Programs/UnrealVS/app.config
new file mode 100644
--- /dev/null
+++ b/Engine/Source/Programs/UnrealVS/app.config
@@ -0,0 +1,27 @@
+﻿<?xml version="1.0" encoding="utf-8"?>
+<configuration>
+  <runtime>
+    <assemblyBinding xmlns="urn:schemas-microsoft-com:asm.v1">
+      <dependentAssembly>
+        <assemblyIdentity name="Microsoft.VisualStudio.Shell.Framework" publicKeyToken="b03f5f7f11d50a3a" culture="neutral" />
+        <bindingRedirect oldVersion="0.0.0.0-16.0.0.0" newVersion="16.0.0.0" />
+      </dependentAssembly>
+      <dependentAssembly>
+        <assemblyIdentity name="System.Threading.Tasks.Extensions" publicKeyToken="cc7b13ffcd2ddd51" culture="neutral" />
+        <bindingRedirect oldVersion="0.0.0.0-4.2.0.1" newVersion="4.2.0.1" />
+      </dependentAssembly>
+      <dependentAssembly>
+        <assemblyIdentity name="System.Runtime.CompilerServices.Unsafe" publicKeyToken="b03f5f7f11d50a3a" culture="neutral" />
+        <bindingRedirect oldVersion="0.0.0.0-4.0.6.0" newVersion="4.0.6.0" />
+      </dependentAssembly>
+      <dependentAssembly>
+        <assemblyIdentity name="System.ValueTuple" publicKeyToken="cc7b13ffcd2ddd51" culture="neutral" />
+        <bindingRedirect oldVersion="0.0.0.0-4.0.3.0" newVersion="4.0.3.0" />
+      </dependentAssembly>
+      <dependentAssembly>
+        <assemblyIdentity name="System.Buffers" publicKeyToken="cc7b13ffcd2ddd51" culture="neutral" />
+        <bindingRedirect oldVersion="0.0.0.0-4.0.3.0" newVersion="4.0.3.0" />
+      </dependentAssembly>
+    </assemblyBinding>
+  </runtime>
+</configuration>
\ No newline at end of file
