UnrealVS - Fixed crash opening some net core projects that had multiple launch profiles, but had not choosen which to use.

diff --git a/Engine/Source/Programs/UnrealVS/CommandLineEditor.cs b/Engine/Source/Programs/UnrealVS/CommandLineEditor.cs
--- a/Engine/Source/Programs/UnrealVS/CommandLineEditor.cs
+++ b/Engine/Source/Programs/UnrealVS/CommandLineEditor.cs
@@ -225,7 +225,7 @@ namespace UnrealVS
 										PropertyNameCaseInsensitive = true
 									});
 
-									if (activeDebugProfile != null)
+									if (!string.IsNullOrEmpty(activeDebugProfile))
 									{
 										Text = settings.Profiles[activeDebugProfile].CommandLineArgs ?? "";
 									}
@@ -424,7 +424,7 @@ namespace UnrealVS
 											PropertyNameCaseInsensitive = true
 										});
 
-										if (activeDebugProfile != null)
+										if (!string.IsNullOrEmpty(activeDebugProfile))
 										{
 											settings.Profiles[activeDebugProfile].CommandLineArgs = FullCommandLine;
 										}
