Fixed copyright notices

diff --git a/Engine/Source/Programs/UnrealVS/ProvideSolutionProperties.cs b/Engine/Source/Programs/UnrealVS/ProvideSolutionProperties.cs
--- a/Engine/Source/Programs/UnrealVS/ProvideSolutionProperties.cs
+++ b/Engine/Source/Programs/UnrealVS/ProvideSolutionProperties.cs
@@ -1,3 +1,5 @@
+// Copyright Epic Games, Inc. All Rights Reserved.
+
 using System;
 using System.Globalization;
 using Microsoft.VisualStudio.Shell;
