Added packages.config for UnrealVS

diff --git a/Engine/Source/Programs/UnrealVS/packages.config b/Engine/Source/Programs/UnrealVS/packages.config
new file mode 100644
--- /dev/null
+++ b/Engine/Source/Programs/UnrealVS/packages.config
@@ -0,0 +1,18 @@
+﻿<?xml version="1.0" encoding="utf-8"?>
+<packages>
+  <package id="Microsoft.Bcl.AsyncInterfaces" version="1.1.0" targetFramework="net472" />
+  <package id="Microsoft.Win32.Registry" version="4.6.0" targetFramework="net472" />
+  <package id="Microsoft.VisualStudio.Threading" version="16.7.56" targetFramework="net472" />
+  <package id="Microsoft.VisualStudio.Threading.Analyzers" version="16.7.56" targetFramework="net472" />
+  <package id="Microsoft.VisualStudio.Validation" version="15.5.31" targetFramework="net472" />
+  <package id="System.Buffers" version="4.5.1" targetFramework="net472" />
+  <package id="System.Memory" version="4.5.4" targetFramework="net472" />
+  <package id="System.Numerics.Vectors" version="4.5.0" targetFramework="net472" />
+  <package id="System.Runtime.CompilerServices.Unsafe" version="4.7.1" targetFramework="net472" />
+  <package id="System.Security.AccessControl" version="4.6.0" targetFramework="net472" />
+  <package id="System.Security.Principal.Windows" version="4.6.0" targetFramework="net472" />
+  <package id="System.Text.Encodings.Web" version="4.7.1" targetFramework="net472" />
+  <package id="System.Text.Json" version="4.7.2" targetFramework="net472" />
+  <package id="System.Threading.Tasks.Extensions" version="4.5.4" targetFramework="net472" />
+  <package id="System.ValueTuple" version="4.5.0" targetFramework="net472" />
+</packages>
\ No newline at end of file
