UnrealVS: Handle missing launchSettings.json for C# project and add additional error checking
Update nuget reference versions to fix version conflicts
Minor format document changes
Bump version to 1.58

diff --git a/Engine/Source/Programs/UnrealVS/CommandLineEditor.cs b/Engine/Source/Programs/UnrealVS/CommandLineEditor.cs
--- a/Engine/Source/Programs/UnrealVS/CommandLineEditor.cs
+++ b/Engine/Source/Programs/UnrealVS/CommandLineEditor.cs
@@ -88,10 +88,9 @@ namespace UnrealVS
 			UnrealVSPackage.Instance.SolutionBuildManager.get_StartupProject(out ProjectHierarchy);
 			if (ProjectHierarchy != null)
 			{
-				
 				// @todo: filter this so that we only respond to changes in the command line property
 				// Setup a timer to prevent any more performance problem from spamming events
-				if( UpdateCommandLineComboTimer == null )
+				if (UpdateCommandLineComboTimer == null)
 				{
 					UpdateCommandLineComboTimer = new System.Timers.Timer(1000);
 					UpdateCommandLineComboTimer.AutoReset = false;
@@ -108,7 +107,7 @@ namespace UnrealVS
 		/// </summary>
 		private void OnUpdateCommandLineCombo(Object source, ElapsedEventArgs e)
 		{
-            ThreadHelper.Generic.BeginInvoke(UpdateCommandLineCombo);
+			ThreadHelper.Generic.BeginInvoke(UpdateCommandLineCombo);
 		}
 
 		/// <summary>
@@ -159,10 +158,10 @@ namespace UnrealVS
 		private void UpdateCommandLineCombo()
 		{
 			// Enable or disable our command-line selector
-			DesiredCommandLine = null;	// clear this state var used by the combo box handler
+			DesiredCommandLine = null;  // clear this state var used by the combo box handler
 			IVsHierarchy ProjectHierarchy;
-			UnrealVSPackage.Instance.SolutionBuildManager.get_StartupProject( out ProjectHierarchy );
-			if( ProjectHierarchy != null )
+			UnrealVSPackage.Instance.SolutionBuildManager.get_StartupProject(out ProjectHierarchy);
+			if (ProjectHierarchy != null)
 			{
 				ComboCommand.Enabled = true;
 				ComboCommandList.Enabled = true;
@@ -211,7 +210,7 @@ namespace UnrealVS
 							{
 								// net core projects move debugger arguments into a debug profile setup via launchSetings.json
 								string activeDebugProfile;
-								if (PropertyStorage.GetPropertyValue("ActiveDebugProfile", ConfigurationName, (uint) _PersistStorageType.PST_USER_FILE, out activeDebugProfile) != VSConstants.S_OK)
+								if (PropertyStorage.GetPropertyValue("ActiveDebugProfile", ConfigurationName, (uint)_PersistStorageType.PST_USER_FILE, out activeDebugProfile) != VSConstants.S_OK)
 								{
 									activeDebugProfile = null;
 								}
@@ -348,7 +347,7 @@ namespace UnrealVS
 
 							// Get the project platform name.
 							string ProjectPlatformName = SelectedConfiguration.PlatformName;
-							if(ProjectPlatformName == "Any CPU")
+							if (ProjectPlatformName == "Any CPU")
 							{
 								ProjectPlatformName = "AnyCPU";
 							}
@@ -366,7 +365,7 @@ namespace UnrealVS
 								{
 									PropertyStorage.RemoveProperty("LocalDebuggerCommandArguments", ProjectConfigurationName, (uint)_PersistStorageType.PST_USER_FILE);
 									PropertyStorage.RemoveProperty("RemoteDebuggerCommandArguments", ProjectConfigurationName, (uint)_PersistStorageType.PST_USER_FILE);
-									foreach( string ExtraField in ExtraFields )
+									foreach (string ExtraField in ExtraFields)
 									{
 										PropertyStorage.RemoveProperty(ExtraField, ProjectConfigurationName, (uint)_PersistStorageType.PST_USER_FILE);
 									}
@@ -416,30 +415,53 @@ namespace UnrealVS
 									}
 
 									string launchSettingsPath = Path.Combine(Path.GetDirectoryName(SelectedStartupProject.FileName), "Properties", "launchSettings.json");
+									LaunchSettingsJson settings = null;
 									if (File.Exists(launchSettingsPath))
 									{
-										string jsonDoc = File.ReadAllText(launchSettingsPath);
-										LaunchSettingsJson settings = JsonSerializer.Deserialize<LaunchSettingsJson>(jsonDoc, new JsonSerializerOptions
-										{
-											PropertyNameCaseInsensitive = true
-										});
-
-										if (!string.IsNullOrEmpty(activeDebugProfile))
+										try
 										{
-											settings.Profiles[activeDebugProfile].CommandLineArgs = FullCommandLine;
+											string jsonDoc = File.ReadAllText(launchSettingsPath);
+											settings = JsonSerializer.Deserialize<LaunchSettingsJson>(jsonDoc, new JsonSerializerOptions
+											{
+												PropertyNameCaseInsensitive = true
+											});
 										}
-										else
+										catch (Exception Ex)
 										{
-											// if no active debug profile is set then VS uses the first one
-											settings.Profiles.First().Value.CommandLineArgs = FullCommandLine;
+											Logging.WriteLine($"Error loading {launchSettingsPath}");
+											Logging.WriteLine(Ex.ToString());
 										}
+									}
 
-										string json = JsonSerializer.Serialize(settings, new JsonSerializerOptions
+									// Create empty Launch Settings if launchSettings.json does not exist or failed to load
+									if (settings == null || settings.Profiles == null || settings.Profiles.Count == 0)
+									{
+										settings = new LaunchSettingsJson
+										{
+											Profiles = new Dictionary<string, LaunchSettingsJson.LaunchSettingsProfile>()
+										};
+										settings.Profiles.Add(SelectedStartupProject.Name, new LaunchSettingsJson.LaunchSettingsProfile
 										{
-											WriteIndented = true,
+											CommandName = "Project",
+											CommandLineArgs = string.Empty,
 										});
-										File.WriteAllText(launchSettingsPath, json);
 									}
+
+									if (!string.IsNullOrEmpty(activeDebugProfile) && settings.Profiles.ContainsKey(activeDebugProfile))
+									{
+										settings.Profiles[activeDebugProfile].CommandLineArgs = FullCommandLine;
+									}
+									else
+									{
+										// if no active debug profile is set then VS uses the first one
+										settings.Profiles.First().Value.CommandLineArgs = FullCommandLine;
+									}
+
+									string json = JsonSerializer.Serialize(settings, new JsonSerializerOptions
+									{
+										WriteIndented = true,
+									});
+									File.WriteAllText(launchSettingsPath, json);
 								}
 							}
 						}
diff --git a/Engine/Source/Programs/UnrealVS/Logging.cs b/Engine/Source/Programs/UnrealVS/Logging.cs
--- a/Engine/Source/Programs/UnrealVS/Logging.cs
+++ b/Engine/Source/Programs/UnrealVS/Logging.cs
@@ -1,4 +1,4 @@
-﻿// Copyright Epic Games, Inc. All Rights Reserved.
+// Copyright Epic Games, Inc. All Rights Reserved.
 
 using System;
 using System.Collections.Generic;
@@ -69,9 +69,13 @@ namespace UnrealVS
 					WriteLine("Visual Studio 11 build");
 #elif VS12
 					WriteLine("Visual Studio 12 build");
+#elif VS13
+					WriteLine("Visual Studio 13 build");
+#elif VS14
+					WriteLine("Visual Studio 14 build");
 #else
-					WriteLine("UNKNOWN Build");
-#endif				
+					WriteLine("UNKNOWN Visual Studio build");
+#endif
 				}
 				catch (IOException)
 				{
diff --git a/Engine/Source/Programs/UnrealVS/Properties/AssemblyInfo.cs b/Engine/Source/Programs/UnrealVS/Properties/AssemblyInfo.cs
--- a/Engine/Source/Programs/UnrealVS/Properties/AssemblyInfo.cs
+++ b/Engine/Source/Programs/UnrealVS/Properties/AssemblyInfo.cs
@@ -31,8 +31,8 @@ using System.Runtime.InteropServices;
 // You can specify all the values or you can default the Revision and Build Numbers 
 // by using the '*' as shown below:
 
-[assembly: AssemblyVersion("1.57.0.0")]
-[assembly: AssemblyFileVersion("1.57.0.0")]
+[assembly: AssemblyVersion("1.58.0.0")]
+[assembly: AssemblyFileVersion("1.58.0.0")]
 
 
 
diff --git a/Engine/Source/Programs/UnrealVS/app.config b/Engine/Source/Programs/UnrealVS/app.config
--- a/Engine/Source/Programs/UnrealVS/app.config
+++ b/Engine/Source/Programs/UnrealVS/app.config
@@ -12,7 +12,7 @@
       </dependentAssembly>
       <dependentAssembly>
         <assemblyIdentity name="System.Runtime.CompilerServices.Unsafe" publicKeyToken="b03f5f7f11d50a3a" culture="neutral" />
-        <bindingRedirect oldVersion="0.0.0.0-4.0.6.0" newVersion="4.0.6.0" />
+        <bindingRedirect oldVersion="0.0.0.0-5.0.0.0" newVersion="5.0.0.0" />
       </dependentAssembly>
       <dependentAssembly>
         <assemblyIdentity name="System.ValueTuple" publicKeyToken="cc7b13ffcd2ddd51" culture="neutral" />
diff --git a/Engine/Source/Programs/UnrealVS/packages.config b/Engine/Source/Programs/UnrealVS/packages.config
--- a/Engine/Source/Programs/UnrealVS/packages.config
+++ b/Engine/Source/Programs/UnrealVS/packages.config
@@ -1,18 +1,18 @@
 ﻿<?xml version="1.0" encoding="utf-8"?>
 <packages>
-  <package id="Microsoft.Bcl.AsyncInterfaces" version="1.1.0" targetFramework="net472" />
-  <package id="Microsoft.Win32.Registry" version="4.6.0" targetFramework="net472" />
-  <package id="Microsoft.VisualStudio.Threading" version="16.7.56" targetFramework="net472" />
-  <package id="Microsoft.VisualStudio.Threading.Analyzers" version="16.7.56" targetFramework="net472" />
-  <package id="Microsoft.VisualStudio.Validation" version="15.5.31" targetFramework="net472" />
+  <package id="Microsoft.Bcl.AsyncInterfaces" version="5.0.0" targetFramework="net472" />
+  <package id="Microsoft.VisualStudio.Threading" version="16.9.60" targetFramework="net472" />
+  <package id="Microsoft.VisualStudio.Threading.Analyzers" version="16.9.60" targetFramework="net472" developmentDependency="true" />
+  <package id="Microsoft.VisualStudio.Validation" version="16.9.32" targetFramework="net472" />
+  <package id="Microsoft.Win32.Registry" version="5.0.0" targetFramework="net472" />
   <package id="System.Buffers" version="4.5.1" targetFramework="net472" />
   <package id="System.Memory" version="4.5.4" targetFramework="net472" />
   <package id="System.Numerics.Vectors" version="4.5.0" targetFramework="net472" />
-  <package id="System.Runtime.CompilerServices.Unsafe" version="4.7.1" targetFramework="net472" />
-  <package id="System.Security.AccessControl" version="4.6.0" targetFramework="net472" />
-  <package id="System.Security.Principal.Windows" version="4.6.0" targetFramework="net472" />
-  <package id="System.Text.Encodings.Web" version="4.7.1" targetFramework="net472" />
-  <package id="System.Text.Json" version="4.7.2" targetFramework="net472" />
+  <package id="System.Runtime.CompilerServices.Unsafe" version="5.0.0" targetFramework="net472" />
+  <package id="System.Security.AccessControl" version="5.0.0" targetFramework="net472" />
+  <package id="System.Security.Principal.Windows" version="5.0.0" targetFramework="net472" />
+  <package id="System.Text.Encodings.Web" version="5.0.1" targetFramework="net472" />
+  <package id="System.Text.Json" version="5.0.2" targetFramework="net472" />
   <package id="System.Threading.Tasks.Extensions" version="4.5.4" targetFramework="net472" />
   <package id="System.ValueTuple" version="4.5.0" targetFramework="net472" />
 </packages>
\ No newline at end of file
