UnrealVS: Add .cc and .cxx to file extensions that can be single compiled

diff --git a/Engine/Source/Programs/UnrealVS/CompileSingleFile.cs b/Engine/Source/Programs/UnrealVS/CompileSingleFile.cs
--- a/Engine/Source/Programs/UnrealVS/CompileSingleFile.cs
+++ b/Engine/Source/Programs/UnrealVS/CompileSingleFile.cs
@@ -1,17 +1,11 @@
 // Copyright Epic Games, Inc. All Rights Reserved.
 
-using System;
-using System.Diagnostics;
-using System.ComponentModel.Design;
-using Microsoft.VisualStudio;
-using Microsoft.VisualStudio.Shell;
-using Microsoft.VisualStudio.Shell.Interop;
-using Microsoft.VisualStudio.OLE.Interop;
 using EnvDTE;
-using System.Runtime.InteropServices;
+using Microsoft.VisualStudio.Shell.Interop;
+using System;
 using System.Collections.Generic;
-using EnvDTE80;
-using System.Linq;
+using System.ComponentModel.Design;
+using System.Diagnostics;
 using System.IO;
 using System.Windows.Forms;
 
@@ -20,14 +14,15 @@ namespace UnrealVS
 	class CompileSingleFile : IDisposable
 	{
 		const int CompileSingleFileButtonID = 0x1075;
+		static readonly List<string> ValidExtensions = new List<string> { ".c", ".cc", ".cpp", ".cxx" };
 
 		System.Diagnostics.Process ChildProcess;
 
 		public CompileSingleFile()
 		{
-			CommandID CommandID = new CommandID( GuidList.UnrealVSCmdSet, CompileSingleFileButtonID );
-			MenuCommand CompileSingleFileButtonCommand = new MenuCommand( new EventHandler( CompileSingleFileButtonHandler ), CommandID );
-			UnrealVSPackage.Instance.MenuCommandService.AddCommand( CompileSingleFileButtonCommand );
+			CommandID CommandID = new CommandID(GuidList.UnrealVSCmdSet, CompileSingleFileButtonID);
+			MenuCommand CompileSingleFileButtonCommand = new MenuCommand(new EventHandler(CompileSingleFileButtonHandler), CommandID);
+			UnrealVSPackage.Instance.MenuCommandService.AddCommand(CompileSingleFileButtonCommand);
 		}
 
 		public void Dispose()
@@ -35,9 +30,9 @@ namespace UnrealVS
 			KillChildProcess();
 		}
 
-		void CompileSingleFileButtonHandler( object Sender, EventArgs Args )
+		void CompileSingleFileButtonHandler(object Sender, EventArgs Args)
 		{
-			if(!TryCompileSingleFile())
+			if (!TryCompileSingleFile())
 			{
 				DTE DTE = UnrealVSPackage.Instance.DTE;
 				DTE.ExecuteCommand("Build.Compile");
@@ -46,9 +41,9 @@ namespace UnrealVS
 
 		void KillChildProcess()
 		{
-			if(ChildProcess != null)
+			if (ChildProcess != null)
 			{
-				if(!ChildProcess.HasExited)
+				if (!ChildProcess.HasExited)
 				{
 					ChildProcess.Kill();
 					ChildProcess.WaitForExit();
@@ -65,37 +60,30 @@ namespace UnrealVS
 			// Activate the output window
 			Window Window = DTE.Windows.Item(EnvDTE.Constants.vsWindowKindOutput);
 			Window.Activate();
-			OutputWindow OutputWindow = Window.Object as OutputWindow;
 
-			// Try to find the 'Build' window
-			OutputWindowPane BuildOutputPane = null;
-			foreach(OutputWindowPane Pane in OutputWindow.OutputWindowPanes)
-			{
-				if(Pane.Guid.ToUpperInvariant() == "{1BD8A850-02D1-11D1-BEE7-00A0C913D1F8}")
-				{
-					BuildOutputPane = Pane;
-					break;
-				}
-			}
-			if(BuildOutputPane == null)
+			// Find or create the 'Build' window
+			IVsOutputWindowPane BuildOutputPane = UnrealVSPackage.Instance.GetOutputPane();
+			if (BuildOutputPane == null)
 			{
+				Logging.WriteLine("CompileSingleFile: Build Output Pane not found");
 				return false;
 			}
 
 			// If there's already a build in progress, offer to cancel it
-			if(ChildProcess != null && !ChildProcess.HasExited)
+			if (ChildProcess != null && !ChildProcess.HasExited)
 			{
-				if(MessageBox.Show("Cancel current compile?", "Compile in progress", MessageBoxButtons.YesNo) == DialogResult.Yes)
+				if (MessageBox.Show("Cancel current compile?", "Compile in progress", MessageBoxButtons.YesNo) == DialogResult.Yes)
 				{
 					KillChildProcess();
-					BuildOutputPane.OutputString("1>  Build cancelled.\r\n");
+					BuildOutputPane.OutputString($"1>  Build cancelled.{Environment.NewLine}");
 				}
 				return true;
 			}
 
 			// Check we've got a file open
-			if(DTE.ActiveDocument == null)
+			if (DTE.ActiveDocument == null)
 			{
+				Logging.WriteLine("CompileSingleFile: ActiveDocument not found");
 				return false;
 			}
 
@@ -104,43 +92,49 @@ namespace UnrealVS
 			UnrealVSPackage.Instance.SolutionBuildManager.get_StartupProject(out ProjectHierarchy);
 			if (ProjectHierarchy == null)
 			{
+				Logging.WriteLine("CompileSingleFile: ProjectHierarchy not found");
 				return false;
 			}
 			Project StartupProject = Utils.HierarchyObjectToProject(ProjectHierarchy);
 			if (StartupProject == null)
 			{
+				Logging.WriteLine("CompileSingleFile: StartupProject not found");
 				return false;
 			}
 			Microsoft.VisualStudio.VCProjectEngine.VCProject VCStartupProject = StartupProject.Object as Microsoft.VisualStudio.VCProjectEngine.VCProject;
-			if(VCStartupProject == null)
+			if (VCStartupProject == null)
 			{
+				Logging.WriteLine("CompileSingleFile: VCStartupProject not found");
 				return false;
 			}
 
 			// Get the active configuration for the startup project
 			Configuration ActiveConfiguration = StartupProject.ConfigurationManager.ActiveConfiguration;
-			string ActiveConfigurationName = String.Format("{0}|{1}", ActiveConfiguration.ConfigurationName, ActiveConfiguration.PlatformName);
+			string ActiveConfigurationName = $"{ActiveConfiguration.ConfigurationName}|{ActiveConfiguration.PlatformName}";
 			Microsoft.VisualStudio.VCProjectEngine.VCConfiguration ActiveVCConfiguration = VCStartupProject.Configurations.Item(ActiveConfigurationName);
-			if(ActiveVCConfiguration == null)
+			if (ActiveVCConfiguration == null)
 			{
+				Logging.WriteLine("CompileSingleFile: VCStartupProject ActiveConfiguration not found");
 				return false;
 			}
 
 			// Get the NMake settings for this configuration
 			Microsoft.VisualStudio.VCProjectEngine.VCNMakeTool ActiveNMakeTool = ActiveVCConfiguration.Tools.Item("VCNMakeTool");
-			if(ActiveNMakeTool == null)
+			if (ActiveNMakeTool == null)
 			{
+				MessageBox.Show($"No NMakeTool set for Project {VCStartupProject.Name} set for single-file compile.", "NMakeTool not set", MessageBoxButtons.OK);
 				return false;
 			}
 
 			// Save all the open documents
 			DTE.ExecuteCommand("File.SaveAll");
 
-			// Check it's a cpp file
+			// Check if the requested file is valid
 			string FileToCompile = DTE.ActiveDocument.FullName;
-			if (!FileToCompile.EndsWith(".c", StringComparison.InvariantCultureIgnoreCase) && !FileToCompile.EndsWith(".cpp", StringComparison.InvariantCultureIgnoreCase))
+			string FileToCompileExt = Path.GetExtension(FileToCompile);
+			if (!ValidExtensions.Contains(FileToCompileExt.ToLowerInvariant()))
 			{
-				MessageBox.Show("Invalid file extension for single-file compile.", "Invalid Extension", MessageBoxButtons.OK);
+				MessageBox.Show($"Invalid file extension {FileToCompileExt} for single-file compile.", "Invalid Extension", MessageBoxButtons.OK);
 				return true;
 			}
 
@@ -160,14 +154,14 @@ namespace UnrealVS
 			// Set up the output pane
 			BuildOutputPane.Activate();
 			BuildOutputPane.Clear();
-			BuildOutputPane.OutputString(String.Format("1>------ Build started: Project: {0}, Configuration: {1} {2} ------\r\n", StartupProject.Name, ActiveConfiguration.ConfigurationName, ActiveConfiguration.PlatformName));
-			BuildOutputPane.OutputString(String.Format("1>  Compiling {0}\r\n", FileToCompile));
+			BuildOutputPane.OutputString($"1>------ Build started: Project: {StartupProject.Name}, Configuration: {ActiveConfiguration.ConfigurationName} {ActiveConfiguration.PlatformName} ------{Environment.NewLine}");
+			BuildOutputPane.OutputString($"1>  Compiling {FileToCompile}{Environment.NewLine}");
 
 			// Set up event handlers 
 			DTE.Events.BuildEvents.OnBuildBegin += BuildEvents_OnBuildBegin;
 
 			// Create a delegate for handling output messages
-			DataReceivedEventHandler OutputHandler = (Sender, Args) => { if(Args.Data != null) BuildOutputPane.OutputString("1>  " + Args.Data + "\r\n"); };
+			DataReceivedEventHandler OutputHandler = (Sender, Args) => { if (Args.Data != null) BuildOutputPane.OutputString($"1>  {Args.Data}{Environment.NewLine}"); };
 
 			// Get the build command line and escape any environment variables that we use
 			string BuildCommandLine = ActiveNMakeTool.BuildCommandLine;
@@ -177,7 +171,7 @@ namespace UnrealVS
 			// Spawn the new process
 			ChildProcess = new System.Diagnostics.Process();
 			ChildProcess.StartInfo.FileName = Path.Combine(Environment.SystemDirectory, "cmd.exe");
-			ChildProcess.StartInfo.Arguments = String.Format("/C \"{0} -singlefile=\"{1}\"\"", BuildCommandLine, FileToCompile);
+			ChildProcess.StartInfo.Arguments = $"/C \"{BuildCommandLine} -singlefile=\"{FileToCompile}\"\"";
 			ChildProcess.StartInfo.WorkingDirectory = Path.GetDirectoryName(StartupProject.FullName);
 			ChildProcess.StartInfo.UseShellExecute = false;
 			ChildProcess.StartInfo.RedirectStandardOutput = true;
diff --git a/Engine/Source/Programs/UnrealVS/Properties/AssemblyInfo.cs b/Engine/Source/Programs/UnrealVS/Properties/AssemblyInfo.cs
--- a/Engine/Source/Programs/UnrealVS/Properties/AssemblyInfo.cs
+++ b/Engine/Source/Programs/UnrealVS/Properties/AssemblyInfo.cs
@@ -31,8 +31,8 @@ using System.Runtime.InteropServices;
 // You can specify all the values or you can default the Revision and Build Numbers 
 // by using the '*' as shown below:
 
-[assembly: AssemblyVersion("1.58.0.0")]
-[assembly: AssemblyFileVersion("1.58.0.0")]
+[assembly: AssemblyVersion("1.59.0.0")]
+[assembly: AssemblyFileVersion("1.59.0.0")]
 
 
 
