UnrealVS: Update NuGet to use PackageReference instead of package.config

https://docs.microsoft.com/en-us/nuget/consume-packages/migrate-packages-config-to-package-reference

Fix Code Analysis warnings
Set warnings as errors
Format all documents & remove unused usings

diff --git a/Engine/Source/Programs/UnrealVS/BatchBuilder.cs b/Engine/Source/Programs/UnrealVS/BatchBuilder.cs
--- a/Engine/Source/Programs/UnrealVS/BatchBuilder.cs
+++ b/Engine/Source/Programs/UnrealVS/BatchBuilder.cs
@@ -1,14 +1,10 @@
 // Copyright Epic Games, Inc. All Rights Reserved.
 
-using System;
-using System.Collections.Generic;
-using System.IO;
-using System.Linq;
-using System.Text;
-using System.Threading.Tasks;
-using System.ComponentModel.Design;
 using Microsoft.VisualStudio.Shell;
 using Microsoft.VisualStudio.Shell.Interop;
+using System;
+using System.ComponentModel.Design;
+using System.IO;
 
 namespace UnrealVS
 {
@@ -32,7 +28,7 @@ namespace UnrealVS
 		private static BatchBuilderToolControl _ToolControl;
 		public static BatchBuilderToolControl ToolControl
 		{
-			get 
+			get
 			{
 				if (_ToolControl == null) _ToolControl = new BatchBuilderToolControl(_ToolControlState);
 				return _ToolControl;
diff --git a/Engine/Source/Programs/UnrealVS/BatchBuilderToolControl.xaml.cs b/Engine/Source/Programs/UnrealVS/BatchBuilderToolControl.xaml.cs
--- a/Engine/Source/Programs/UnrealVS/BatchBuilderToolControl.xaml.cs
+++ b/Engine/Source/Programs/UnrealVS/BatchBuilderToolControl.xaml.cs
@@ -1,20 +1,19 @@
 // Copyright Epic Games, Inc. All Rights Reserved.
 
+using EnvDTE;
+using EnvDTE80;
+using Microsoft.VisualStudio;
+using Microsoft.VisualStudio.Shell;
+using Microsoft.VisualStudio.Shell.Interop;
 using System;
 using System.Collections.Generic;
 using System.Collections.ObjectModel;
-using System.Collections.Specialized;
 using System.IO;
 using System.Linq;
 using System.Text;
 using System.Windows;
 using System.Windows.Controls;
 using System.Windows.Input;
-using EnvDTE;
-using EnvDTE80;
-using Microsoft.VisualStudio;
-using Microsoft.VisualStudio.Shell;
-using Microsoft.VisualStudio.Shell.Interop;
 using Constants = EnvDTE.Constants;
 
 namespace UnrealVS
@@ -30,7 +29,7 @@ namespace UnrealVS
 		public class BuildJob : DependencyObject
 		{
 			////////////// BuildJob Types ///////////////////////
-			
+
 			public enum BuildJobType
 			{
 				Build,
@@ -102,25 +101,25 @@ namespace UnrealVS
 
 			public BuildJobStatus JobStatus
 			{
-				get { return (BuildJobStatus) GetValue(JobStatusProperty); }
+				get { return (BuildJobStatus)GetValue(JobStatusProperty); }
 				set { SetValue(JobStatusProperty, value); }
 			}
 
 			public string JobStatusDisplayString
 			{
-				get { return (string) GetValue(JobStatusDisplayStringProperty); }
+				get { return (string)GetValue(JobStatusDisplayStringProperty); }
 				private set { SetValue(JobStatusDisplayStringProperty, value); }
 			}
 
 			public string OutputText
 			{
-				get { return (string) GetValue(OutputTextProperty); }
+				get { return (string)GetValue(OutputTextProperty); }
 				set { SetValue(OutputTextProperty, value); }
 			}
 
 			public bool HasOutputText
 			{
-				get { return (bool) GetValue(HasOutputTextProperty); }
+				get { return (bool)GetValue(HasOutputTextProperty); }
 				private set { SetValue(HasOutputTextProperty, value); }
 			}
 
@@ -130,26 +129,26 @@ namespace UnrealVS
 			{
 				// Register the dependency properties
 				JobStatusProperty = DependencyProperty.Register("JobStatus",
-				                                                typeof (BuildJobStatus), typeof (BuildJob),
-				                                                new FrameworkPropertyMetadata(BuildJobStatus.Pending,
-				                                                                              OnJobStatusChanged));
+																typeof(BuildJobStatus), typeof(BuildJob),
+																new FrameworkPropertyMetadata(BuildJobStatus.Pending,
+																							  OnJobStatusChanged));
 
 				JobStatusDisplayStringProperty = DependencyProperty.Register("JobStatusDisplayString",
-				                                                             typeof (string), typeof (BuildJob),
-				                                                             new FrameworkPropertyMetadata(InvalidJobStatusString));
+																			 typeof(string), typeof(BuildJob),
+																			 new FrameworkPropertyMetadata(InvalidJobStatusString));
 
 				OutputTextProperty = DependencyProperty.Register("OutputText",
-				                                                 typeof (string), typeof (BuildJob),
-				                                                 new FrameworkPropertyMetadata(String.Empty, OnOutputTextChanged));
+																 typeof(string), typeof(BuildJob),
+																 new FrameworkPropertyMetadata(String.Empty, OnOutputTextChanged));
 
 				HasOutputTextProperty = DependencyProperty.Register("HasOutputText",
-				                                                    typeof (bool), typeof (BuildJob),
-				                                                    new FrameworkPropertyMetadata(false));
+																	typeof(bool), typeof(BuildJob),
+																	new FrameworkPropertyMetadata(false));
 			}
 
 			private static void OnJobStatusChanged(DependencyObject Obj, DependencyPropertyChangedEventArgs Args)
 			{
-				BuildJob ThisBuildJob = (BuildJob) Obj;
+				BuildJob ThisBuildJob = (BuildJob)Obj;
 				if (ThisBuildJob != null)
 				{
 					ThisBuildJob.JobStatusDisplayString = ThisBuildJob.GetJobStatusString();
@@ -158,7 +157,7 @@ namespace UnrealVS
 
 			private static void OnOutputTextChanged(DependencyObject Obj, DependencyPropertyChangedEventArgs Args)
 			{
-				BuildJob ThisBuildJob = (BuildJob) Obj;
+				BuildJob ThisBuildJob = (BuildJob)Obj;
 				if (ThisBuildJob != null)
 				{
 					ThisBuildJob.HasOutputText = !String.IsNullOrEmpty(ThisBuildJob.OutputText);
@@ -183,13 +182,13 @@ namespace UnrealVS
 			{
 				if (Project == null || Config == null || Platform == null) return InvalidDisplayString;
 
-				return String.Format("{0}: {1} [{2}]", Enum.GetName(typeof (BuildJobType), JobType), Project.Name,
-				                     SolutionConfigString);
+				return String.Format("{0}: {1} [{2}]", Enum.GetName(typeof(BuildJobType), JobType), Project.Name,
+									 SolutionConfigString);
 			}
 
 			private string GetJobStatusString()
 			{
-				return Enum.GetName(typeof (BuildJobStatus), JobStatus);
+				return Enum.GetName(typeof(BuildJobStatus), JobStatus);
 			}
 
 			////////////// BuildJob fields ///////////////////////
@@ -233,7 +232,7 @@ namespace UnrealVS
 				return "<None>";
 			}
 		}
-		
+
 		public class BatchBuilderToolState
 		{
 			public readonly ObservableCollection<BuildJobSet> _BuildJobSetsCollection = new ObservableCollection<BuildJobSet>();
@@ -254,8 +253,10 @@ namespace UnrealVS
 
 						for (int SetIdx = 0; SetIdx < SetCount; SetIdx++)
 						{
-							BuildJobSet LoadedSet = new BuildJobSet();
-							LoadedSet.Name = Reader.ReadString();
+							BuildJobSet LoadedSet = new BuildJobSet
+							{
+								Name = Reader.ReadString()
+							};
 							int JobCount = Reader.ReadInt32();
 							for (int JobIdx = 0; JobIdx < JobCount; JobIdx++)
 							{
@@ -263,9 +264,8 @@ namespace UnrealVS
 
 								string Config = Reader.ReadString();
 								string Platform = Reader.ReadString();
-								BuildJob.BuildJobType JobType;
 
-								if (Enum.TryParse(Reader.ReadString(), out JobType))
+								if (Enum.TryParse(Reader.ReadString(), out BuildJob.BuildJobType JobType))
 								{
 									LoadedSet.BuildJobs.Add(new BuildJob(ProjectRef, Config, Platform, JobType));
 								}
@@ -343,19 +343,19 @@ namespace UnrealVS
 
 		public bool IsSolutionOpen
 		{
-			get { return (bool) GetValue(IsSolutionOpenProperty); }
+			get { return (bool)GetValue(IsSolutionOpenProperty); }
 			set { SetValue(IsSolutionOpenProperty, value); }
 		}
 
 		public bool IsSingleBuildJobSelected
 		{
-			get { return (bool) GetValue(IsSingleBuildJobSelectedProperty); }
+			get { return (bool)GetValue(IsSingleBuildJobSelectedProperty); }
 			set { SetValue(IsSingleBuildJobSelectedProperty, value); }
 		}
 
 		public ObservableCollection<BuildJob> BuildJobs
 		{
-			get { return (ObservableCollection<BuildJob>) GetValue(BuildJobsProperty); }
+			get { return (ObservableCollection<BuildJob>)GetValue(BuildJobsProperty); }
 			set { SetValue(BuildJobsProperty, value); }
 		}
 
@@ -367,31 +367,31 @@ namespace UnrealVS
 
 		public bool IsDeletableSetSelected
 		{
-			get { return (bool) GetValue(IsDeletableSetSelectedProperty); }
+			get { return (bool)GetValue(IsDeletableSetSelectedProperty); }
 			set { SetValue(IsDeletableSetSelectedProperty, value); }
 		}
 
 		public string BuildJobsPanelTitle
 		{
-			get { return (string) GetValue(BuildJobsPanelTitleProperty); }
+			get { return (string)GetValue(BuildJobsPanelTitleProperty); }
 			set { SetValue(BuildJobsPanelTitleProperty, value); }
 		}
 
 		public string OutputPanelTitle
 		{
-			get { return (string) GetValue(OutputPanelTitleProperty); }
+			get { return (string)GetValue(OutputPanelTitleProperty); }
 			set { SetValue(OutputPanelTitleProperty, value); }
 		}
 
 		public bool HasOutput
 		{
-			get { return (bool) GetValue(HasOutputProperty); }
+			get { return (bool)GetValue(HasOutputProperty); }
 			set { SetValue(HasOutputProperty, value); }
 		}
 
 		public bool IsBusy
 		{
-			get { return (bool) GetValue(IsBusyProperty); }
+			get { return (bool)GetValue(IsBusyProperty); }
 			set { SetValue(IsBusyProperty, value); }
 		}
 
@@ -447,37 +447,37 @@ namespace UnrealVS
 		{
 			// Register the dependency properties
 			IsSolutionOpenProperty = DependencyProperty.Register("IsSolutionOpen",
-			                                                     typeof (bool), typeof (BatchBuilderToolControl),
-			                                                     new FrameworkPropertyMetadata(false));
+																 typeof(bool), typeof(BatchBuilderToolControl),
+																 new FrameworkPropertyMetadata(false));
 
 			IsSingleBuildJobSelectedProperty = DependencyProperty.Register("IsSingleBuildJobSelected",
-			                                                               typeof (bool), typeof (BatchBuilderToolControl),
-			                                                               new FrameworkPropertyMetadata(false));
+																		   typeof(bool), typeof(BatchBuilderToolControl),
+																		   new FrameworkPropertyMetadata(false));
 
 			BuildJobsProperty = DependencyProperty.Register("BuildJobs",
-			                                                typeof (ObservableCollection<BuildJob>),
-			                                                typeof (BatchBuilderToolControl),
-			                                                new FrameworkPropertyMetadata(new ObservableCollection<BuildJob>()));
+															typeof(ObservableCollection<BuildJob>),
+															typeof(BatchBuilderToolControl),
+															new FrameworkPropertyMetadata(new ObservableCollection<BuildJob>()));
 
 			IsDeletableSetSelectedProperty = DependencyProperty.Register("IsDeletableSetSelected",
-			                                                             typeof (bool), typeof (BatchBuilderToolControl),
-			                                                             new FrameworkPropertyMetadata(false));
+																		 typeof(bool), typeof(BatchBuilderToolControl),
+																		 new FrameworkPropertyMetadata(false));
 
 			BuildJobsPanelTitleProperty = DependencyProperty.Register("BuildJobsPanelTitle",
-			                                                          typeof (string), typeof (BatchBuilderToolControl),
-			                                                          new FrameworkPropertyMetadata(BuildJobsPanelPrefix));
+																	  typeof(string), typeof(BatchBuilderToolControl),
+																	  new FrameworkPropertyMetadata(BuildJobsPanelPrefix));
 
 			OutputPanelTitleProperty = DependencyProperty.Register("OutputPanelTitle",
-			                                                       typeof (string), typeof (BatchBuilderToolControl),
-			                                                       new FrameworkPropertyMetadata(OutputPanelTitlePrefix));
+																   typeof(string), typeof(BatchBuilderToolControl),
+																   new FrameworkPropertyMetadata(OutputPanelTitlePrefix));
 
 			HasOutputProperty = DependencyProperty.Register("HasOutput",
-			                                                typeof (bool), typeof (BatchBuilderToolControl),
-			                                                new FrameworkPropertyMetadata(false));
+															typeof(bool), typeof(BatchBuilderToolControl),
+															new FrameworkPropertyMetadata(false));
 
 			IsBusyProperty = DependencyProperty.Register("IsBusy",
-			                                             typeof (bool), typeof (BatchBuilderToolControl),
-			                                             new FrameworkPropertyMetadata(false));
+														 typeof(bool), typeof(BatchBuilderToolControl),
+														 new FrameworkPropertyMetadata(false));
 		}
 
 		private static void DisplayBatchOutputText(string Text)
@@ -517,15 +517,15 @@ namespace UnrealVS
 			// Register for events that we care about
 			UnrealVSPackage.Instance.StartupProjectSelector.StartupProjectListChanged += RefreshProjectCollection;
 			UnrealVSPackage.Instance.OnSolutionOpened += delegate
-				                                             {
-					                                             IsSolutionOpen = true;
-					                                             RefreshConfigAndPlatformCollections();
-				                                             };
+															 {
+																 IsSolutionOpen = true;
+																 RefreshConfigAndPlatformCollections();
+															 };
 			UnrealVSPackage.Instance.OnSolutionClosed += delegate
-				                                             {
-					                                             IsSolutionOpen = false;
-					                                             RefreshConfigAndPlatformCollections();
-				                                             };
+															 {
+																 IsSolutionOpen = false;
+																 RefreshConfigAndPlatformCollections();
+															 };
 
 			RefreshProjectCollection(UnrealVSPackage.Instance.StartupProjectSelector.StartupProjectList);
 			RefreshConfigAndPlatformCollections();
@@ -561,7 +561,7 @@ namespace UnrealVS
 		{
 			if (_state._BuildJobSetsCollection.Count == 0)
 			{
-				BuildJobSet DefaultItem = new BuildJobSet {Name = "UntitledSet"};
+				BuildJobSet DefaultItem = new BuildJobSet { Name = "UntitledSet" };
 				_state._BuildJobSetsCollection.Add(DefaultItem);
 				if (SetCombo != null)
 				{
@@ -595,9 +595,7 @@ namespace UnrealVS
 				_state._BuildJobSetsCollection.Clear();
 			}
 
-			string[] RefreshConfigs;
-			string[] RefreshPlatforms;
-			Utils.GetSolutionConfigsAndPlatforms(out RefreshConfigs, out RefreshPlatforms);
+			Utils.GetSolutionConfigsAndPlatforms(out string[] RefreshConfigs, out string[] RefreshPlatforms);
 
 			foreach (string Config in RefreshConfigs)
 			{
@@ -641,7 +639,7 @@ namespace UnrealVS
 		{
 			if (!IsSingleBuildJobSelected) return;
 
-			BuildJob SelectedJob = (BuildJob) BuildJobsList.SelectedItem;
+			BuildJob SelectedJob = (BuildJob)BuildJobsList.SelectedItem;
 
 			int CurrentIdx = BuildJobs.IndexOf(SelectedJob);
 			if (CurrentIdx > 0)
@@ -654,7 +652,7 @@ namespace UnrealVS
 		{
 			if (!IsSingleBuildJobSelected) return;
 
-			BuildJob SelectedJob = (BuildJob) BuildJobsList.SelectedItem;
+			BuildJob SelectedJob = (BuildJob)BuildJobsList.SelectedItem;
 
 			int CurrentIdx = BuildJobs.IndexOf(SelectedJob);
 			if (CurrentIdx < BuildJobs.Count - 1)
@@ -676,8 +674,8 @@ namespace UnrealVS
 					foreach (string Platform in PlatformsList.SelectedItems)
 					{
 						Buildables.Add(
-							new BuildJob(new Utils.SafeProjectReference {FullName = Project.Project.FullName, Name = Project.Project.Name},
-							             Config, Platform, JobType));
+							new BuildJob(new Utils.SafeProjectReference { FullName = Project.Project.FullName, Name = Project.Project.Name },
+										 Config, Platform, JobType));
 					}
 				}
 			}
@@ -761,7 +759,7 @@ namespace UnrealVS
 		{
 			if (SetCombo.SelectedItem != null)
 			{
-				_LastSelectedBuildJobSet = (BuildJobSet) SetCombo.SelectedItem;
+				_LastSelectedBuildJobSet = (BuildJobSet)SetCombo.SelectedItem;
 				BuildJobs = _LastSelectedBuildJobSet.BuildJobs;
 				IsDeletableSetSelected = _state._BuildJobSetsCollection.Count > 1;
 				BuildJobsPanelTitle = String.Format("{0} ({1})", BuildJobsPanelPrefix, _LastSelectedBuildJobSet.Name);
@@ -845,11 +843,11 @@ namespace UnrealVS
 			if (BuildJobCount > BuildJobs.Count)
 			{
 				if (VSConstants.S_OK != VsShellUtilities.ShowMessageBox(ServiceProvider.GlobalProvider,
-				                                                        "WARNING: Some build jobs were deleted because they were invalid.",
-				                                                        "UnrealVS",
-				                                                        OLEMSGICON.OLEMSGICON_WARNING,
-				                                                        OLEMSGBUTTON.OLEMSGBUTTON_OKCANCEL,
-				                                                        OLEMSGDEFBUTTON.OLEMSGDEFBUTTON_FIRST))
+																		"WARNING: Some build jobs were deleted because they were invalid.",
+																		"UnrealVS",
+																		OLEMSGICON.OLEMSGICON_WARNING,
+																		OLEMSGBUTTON.OLEMSGBUTTON_OKCANCEL,
+																		OLEMSGDEFBUTTON.OLEMSGDEFBUTTON_FIRST))
 				{
 					return;
 				}
@@ -882,9 +880,7 @@ namespace UnrealVS
 		{
 			// check that project, config and platform strings in build jobs are valid
 			// remove invalid ones
-			string[] SolutionConfigs;
-			string[] SolutionPlatforms;
-			Utils.GetSolutionConfigsAndPlatforms(out SolutionConfigs, out SolutionPlatforms);
+			Utils.GetSolutionConfigsAndPlatforms(out string[] SolutionConfigs, out string[] SolutionPlatforms);
 
 			foreach (var JobSet in _state._BuildJobSetsCollection)
 			{
@@ -996,20 +992,17 @@ namespace UnrealVS
 
 		private void OpenItemOutputCommandExecuted(object sender, ExecutedRoutedEventArgs e)
 		{
-			BuildJob Job = e.Parameter as BuildJob;
-			if (Job == null) return;
+			if (!(e.Parameter is BuildJob Job)) return;
 
 			DisplayBatchOutputText(Job.OutputText);
 		}
 
 		private void OnDblClickBuildListItem(object sender, MouseButtonEventArgs e)
 		{
-			var Elem = sender as FrameworkElement;
-			if (Elem == null) return;
+			if (!(sender is FrameworkElement Elem)) return;
+
+			if (!(Elem.DataContext is BuildJob Job)) return;
 
-			BuildJob Job = Elem.DataContext as BuildJob;
-			if (Job == null) return;
-			
 			// Switch the Startup Project and the build config and platform to match this Job
 			if (Job.Project != null)
 			{
@@ -1033,11 +1026,9 @@ namespace UnrealVS
 
 		private void OnDblClickBuildingListItem(object sender, MouseButtonEventArgs e)
 		{
-			var Elem = sender as FrameworkElement;
-			if (Elem == null) return;
+			if (!(sender is FrameworkElement Elem)) return;
 
-			BuildJob Job = Elem.DataContext as BuildJob;
-			if (Job == null) return;
+			if (!(Elem.DataContext is BuildJob Job)) return;
 
 			DisplayBatchOutputText(Job.OutputText);
 		}
diff --git a/Engine/Source/Programs/UnrealVS/BatchBuilderToolWindow.cs b/Engine/Source/Programs/UnrealVS/BatchBuilderToolWindow.cs
--- a/Engine/Source/Programs/UnrealVS/BatchBuilderToolWindow.cs
+++ b/Engine/Source/Programs/UnrealVS/BatchBuilderToolWindow.cs
@@ -1,12 +1,8 @@
-﻿// Copyright Epic Games, Inc. All Rights Reserved.
+// Copyright Epic Games, Inc. All Rights Reserved.
 
+using Microsoft.VisualStudio.Shell;
 using System;
-using System.Collections.Generic;
-using System.Linq;
-using System.Text;
-using System.Threading.Tasks;
 using System.Runtime.InteropServices;
-using Microsoft.VisualStudio.Shell;
 
 namespace UnrealVS
 {
@@ -14,25 +10,25 @@ namespace UnrealVS
 	internal class BatchBuilderToolWindow : ToolWindowPane
 	{
 		/// <summary>
-        /// Standard constructor for the tool window.
-        /// </summary>
+		/// Standard constructor for the tool window.
+		/// </summary>
 		public BatchBuilderToolWindow() :
-            base(null)
-        {
-            // Set the window title reading it from the resources.
-            Caption = Resources.BatchBuilderWindowTitle;
-            // Set the image that will appear on the tab of the window frame
-            // when docked with an other window
-            // The resource ID correspond to the one defined in the resx file
-            // while the Index is the offset in the bitmap strip. Each image in
-            // the strip being 16x16.
-            BitmapResourceID = 200;
-            BitmapIndex = 0;
+			base(null)
+		{
+			// Set the window title reading it from the resources.
+			Caption = Resources.BatchBuilderWindowTitle;
+			// Set the image that will appear on the tab of the window frame
+			// when docked with an other window
+			// The resource ID correspond to the one defined in the resx file
+			// while the Index is the offset in the bitmap strip. Each image in
+			// the strip being 16x16.
+			BitmapResourceID = 200;
+			BitmapIndex = 0;
 
-            // This is the user control hosted by the tool window; Note that, even if this class implements IDisposable,
-            // we are not calling Dispose on this object. This is because ToolWindowPane calls Dispose on 
-            // the object returned by the Content property.
+			// This is the user control hosted by the tool window; Note that, even if this class implements IDisposable,
+			// we are not calling Dispose on this object. This is because ToolWindowPane calls Dispose on 
+			// the object returned by the Content property.
 			base.Content = BatchBuilder.ToolControl;
-        }
+		}
 	}
 }
diff --git a/Engine/Source/Programs/UnrealVS/BuildStartupProject.cs b/Engine/Source/Programs/UnrealVS/BuildStartupProject.cs
--- a/Engine/Source/Programs/UnrealVS/BuildStartupProject.cs
+++ b/Engine/Source/Programs/UnrealVS/BuildStartupProject.cs
@@ -1,16 +1,9 @@
 // Copyright Epic Games, Inc. All Rights Reserved.
 
+using EnvDTE80;
+using Microsoft.VisualStudio.Shell.Interop;
 using System;
-using System.Diagnostics;
 using System.ComponentModel.Design;
-using Microsoft.VisualStudio;
-using Microsoft.VisualStudio.Shell;
-using Microsoft.VisualStudio.Shell.Interop;
-using Microsoft.VisualStudio.OLE.Interop;
-using EnvDTE;
-using System.Runtime.InteropServices;
-using System.Collections.Generic;
-using EnvDTE80;
 
 namespace UnrealVS
 {
@@ -22,29 +15,28 @@ namespace UnrealVS
 		{
 			// BuildStartupProjectButton
 			{
-				var CommandID = new CommandID( GuidList.UnrealVSCmdSet, BuildStartupProjectButtonID );
-				var BuildStartupProjectButtonCommand = new MenuCommand( new EventHandler( BuildStartupProjectButtonHandler ), CommandID );
-				UnrealVSPackage.Instance.MenuCommandService.AddCommand( BuildStartupProjectButtonCommand );
+				var CommandID = new CommandID(GuidList.UnrealVSCmdSet, BuildStartupProjectButtonID);
+				var BuildStartupProjectButtonCommand = new MenuCommand(new EventHandler(BuildStartupProjectButtonHandler), CommandID);
+				UnrealVSPackage.Instance.MenuCommandService.AddCommand(BuildStartupProjectButtonCommand);
 			}
 
 		}
 
 
 		/// Called when 'BuildStartupProject' button is clicked
-		void BuildStartupProjectButtonHandler( object Sender, EventArgs Args )
+		void BuildStartupProjectButtonHandler(object Sender, EventArgs Args)
 		{
 			// Grab the current startup project
-			IVsHierarchy ProjectHierarchy;
-			UnrealVSPackage.Instance.SolutionBuildManager.get_StartupProject( out ProjectHierarchy );
-			if( ProjectHierarchy != null )
+			UnrealVSPackage.Instance.SolutionBuildManager.get_StartupProject(out IVsHierarchy ProjectHierarchy);
+			if (ProjectHierarchy != null)
 			{
-				var StartupProject = Utils.HierarchyObjectToProject( ProjectHierarchy );
+				var StartupProject = Utils.HierarchyObjectToProject(ProjectHierarchy);
 
 				if (StartupProject != null)
 				{
 					// Get the active solution configuration
 					var ActiveConfiguration =
-						(SolutionConfiguration2) UnrealVSPackage.Instance.DTE.Solution.SolutionBuild.ActiveConfiguration;
+						(SolutionConfiguration2)UnrealVSPackage.Instance.DTE.Solution.SolutionBuild.ActiveConfiguration;
 					var SolutionConfiguration = ActiveConfiguration.Name;
 					var SolutionPlatform = ActiveConfiguration.PlatformName;
 
diff --git a/Engine/Source/Programs/UnrealVS/BuildStatusIcon.xaml b/Engine/Source/Programs/UnrealVS/BuildStatusIcon.xaml
--- a/Engine/Source/Programs/UnrealVS/BuildStatusIcon.xaml
+++ b/Engine/Source/Programs/UnrealVS/BuildStatusIcon.xaml
@@ -1,12 +1,12 @@
-﻿<!--Copyright Epic Games, Inc. All Rights Reserved.-->
-	
+<!--Copyright Epic Games, Inc. All Rights Reserved.-->
+
 <UserControl x:Class="UnrealVS.BuildStatusIcon"
-             xmlns="http://schemas.microsoft.com/winfx/2006/xaml/presentation"
-             xmlns:x="http://schemas.microsoft.com/winfx/2006/xaml"
-             xmlns:mc="http://schemas.openxmlformats.org/markup-compatibility/2006" 
-             xmlns:d="http://schemas.microsoft.com/expression/blend/2008" 
-             mc:Ignorable="d" 
-             d:DesignHeight="200" d:DesignWidth="200"
+			 xmlns="http://schemas.microsoft.com/winfx/2006/xaml/presentation"
+			 xmlns:x="http://schemas.microsoft.com/winfx/2006/xaml"
+			 xmlns:mc="http://schemas.openxmlformats.org/markup-compatibility/2006" 
+			 xmlns:d="http://schemas.microsoft.com/expression/blend/2008" 
+			 mc:Ignorable="d" 
+			 d:DesignHeight="200" d:DesignWidth="200"
 			 Foreground="Black">
 	<UserControl.Resources>
 		<LinearGradientBrush x:Key="InnerBrush" EndPoint="0.5,1" StartPoint="0.5,0">
diff --git a/Engine/Source/Programs/UnrealVS/BuildStatusIcon.xaml.cs b/Engine/Source/Programs/UnrealVS/BuildStatusIcon.xaml.cs
--- a/Engine/Source/Programs/UnrealVS/BuildStatusIcon.xaml.cs
+++ b/Engine/Source/Programs/UnrealVS/BuildStatusIcon.xaml.cs
@@ -1,19 +1,8 @@
-﻿// Copyright Epic Games, Inc. All Rights Reserved.
+// Copyright Epic Games, Inc. All Rights Reserved.
 
-using System;
-using System.Collections.Generic;
-using System.Linq;
-using System.Text;
-using System.Threading.Tasks;
 using System.Windows;
 using System.Windows.Controls;
-using System.Windows.Data;
-using System.Windows.Documents;
-using System.Windows.Input;
 using System.Windows.Media;
-using System.Windows.Media.Imaging;
-using System.Windows.Navigation;
-using System.Windows.Shapes;
 
 namespace UnrealVS
 {
@@ -46,11 +35,11 @@ namespace UnrealVS
 		static BuildStatusIcon()
 		{
 			BuildStatusProperty = DependencyProperty.Register("BuildStatus",
-			                                                  typeof (BatchBuilderToolControl.BuildJob.BuildJobStatus),
-			                                                  typeof (BuildStatusIcon),
-			                                                  new FrameworkPropertyMetadata(
-				                                                  BatchBuilderToolControl.BuildJob.BuildJobStatus.Pending,
-				                                                  OnBuildStatusChanged));
+															  typeof(BatchBuilderToolControl.BuildJob.BuildJobStatus),
+															  typeof(BuildStatusIcon),
+															  new FrameworkPropertyMetadata(
+																  BatchBuilderToolControl.BuildJob.BuildJobStatus.Pending,
+																  OnBuildStatusChanged));
 			ColorLoProperty = DependencyProperty.Register("ColorLo",
 															  typeof(Color),
 															  typeof(BuildStatusIcon),
diff --git a/Engine/Source/Programs/UnrealVS/CommandLineEditor.cs b/Engine/Source/Programs/UnrealVS/CommandLineEditor.cs
--- a/Engine/Source/Programs/UnrealVS/CommandLineEditor.cs
+++ b/Engine/Source/Programs/UnrealVS/CommandLineEditor.cs
@@ -1,25 +1,23 @@
 // Copyright Epic Games, Inc. All Rights Reserved.
 
+using EnvDTE;
+using EnvDTE80;
+using Microsoft.VisualStudio;
+using Microsoft.VisualStudio.Shell;
+using Microsoft.VisualStudio.Shell.Interop;
 using System;
-using System.Diagnostics;
+using System.Collections.Generic;
 using System.ComponentModel.Design;
 using System.IO;
 using System.Linq;
-using System.Text.Json;
-using Microsoft.VisualStudio;
-using Microsoft.VisualStudio.Shell;
-using Microsoft.VisualStudio.Shell.Interop;
-using Microsoft.VisualStudio.OLE.Interop;
-using EnvDTE;
 using System.Runtime.InteropServices;
-using System.Collections.Generic;
+using System.Text.Json;
 using System.Text.Json.Serialization;
-using EnvDTE80;
 using System.Timers;
 
 namespace UnrealVS
 {
-	internal class CommandLineEditor
+	internal class CommandLineEditor : IDisposable
 	{
 		private class LaunchSettingsJson
 		{
@@ -84,16 +82,17 @@ namespace UnrealVS
 				return;
 			}
 
-			IVsHierarchy ProjectHierarchy;
-			UnrealVSPackage.Instance.SolutionBuildManager.get_StartupProject(out ProjectHierarchy);
+			UnrealVSPackage.Instance.SolutionBuildManager.get_StartupProject(out IVsHierarchy ProjectHierarchy);
 			if (ProjectHierarchy != null)
 			{
 				// @todo: filter this so that we only respond to changes in the command line property
 				// Setup a timer to prevent any more performance problem from spamming events
 				if (UpdateCommandLineComboTimer == null)
 				{
-					UpdateCommandLineComboTimer = new System.Timers.Timer(1000);
-					UpdateCommandLineComboTimer.AutoReset = false;
+					UpdateCommandLineComboTimer = new System.Timers.Timer(1000)
+					{
+						AutoReset = false,
+					};
 					UpdateCommandLineComboTimer.Elapsed += OnUpdateCommandLineCombo;
 				}
 				// Restart timer to raise the event only after 1s of no notification
@@ -105,7 +104,7 @@ namespace UnrealVS
 		/// <summary>
 		/// Timer callback to UpdateCommandLineCombo
 		/// </summary>
-		private void OnUpdateCommandLineCombo(Object source, ElapsedEventArgs e)
+		private void OnUpdateCommandLineCombo(Object _Source, ElapsedEventArgs _Args)
 		{
 			ThreadHelper.Generic.BeginInvoke(UpdateCommandLineCombo);
 		}
@@ -159,8 +158,7 @@ namespace UnrealVS
 		{
 			// Enable or disable our command-line selector
 			DesiredCommandLine = null;  // clear this state var used by the combo box handler
-			IVsHierarchy ProjectHierarchy;
-			UnrealVSPackage.Instance.SolutionBuildManager.get_StartupProject(out ProjectHierarchy);
+			UnrealVSPackage.Instance.SolutionBuildManager.get_StartupProject(out IVsHierarchy ProjectHierarchy);
 			if (ProjectHierarchy != null)
 			{
 				ComboCommand.Enabled = true;
@@ -181,12 +179,11 @@ namespace UnrealVS
 		/// Returns a string to display in the command-line combo box, based on project state
 		/// </summary>
 		/// <returns>String to display</returns>
-		private string MakeCommandLineComboText()
+		private static string MakeCommandLineComboText()
 		{
 			string Text = "";
 
-			IVsHierarchy ProjectHierarchy;
-			if (UnrealVSPackage.Instance.SolutionBuildManager.get_StartupProject(out ProjectHierarchy) == VSConstants.S_OK && ProjectHierarchy != null)
+			if (UnrealVSPackage.Instance.SolutionBuildManager.get_StartupProject(out IVsHierarchy ProjectHierarchy) == VSConstants.S_OK && ProjectHierarchy != null)
 			{
 				Project SelectedStartupProject = Utils.HierarchyObjectToProject(ProjectHierarchy);
 				if (SelectedStartupProject != null)
@@ -194,8 +191,7 @@ namespace UnrealVS
 					Configuration SelectedConfiguration = SelectedStartupProject.ConfigurationManager.ActiveConfiguration;
 					if (SelectedConfiguration != null)
 					{
-						IVsBuildPropertyStorage PropertyStorage = ProjectHierarchy as IVsBuildPropertyStorage;
-						if (PropertyStorage != null)
+						if (ProjectHierarchy is IVsBuildPropertyStorage PropertyStorage)
 						{
 							// Query the property store for the debugger arguments
 							string ConfigurationName = String.Format("{0}|{1}", SelectedConfiguration.ConfigurationName, SelectedConfiguration.PlatformName);
@@ -209,8 +205,7 @@ namespace UnrealVS
 							else
 							{
 								// net core projects move debugger arguments into a debug profile setup via launchSetings.json
-								string activeDebugProfile;
-								if (PropertyStorage.GetPropertyValue("ActiveDebugProfile", ConfigurationName, (uint)_PersistStorageType.PST_USER_FILE, out activeDebugProfile) != VSConstants.S_OK)
+								if (PropertyStorage.GetPropertyValue("ActiveDebugProfile", ConfigurationName, (uint)_PersistStorageType.PST_USER_FILE, out string activeDebugProfile) != VSConstants.S_OK)
 								{
 									activeDebugProfile = null;
 								}
@@ -264,8 +259,7 @@ namespace UnrealVS
 			{
 				var OleArgs = (OleMenuCmdEventArgs)Args;
 
-				string InString = OleArgs.InValue as string;
-				if (InString != null)
+				if (OleArgs.InValue is string InString)
 				{
 					// New text set on the combo - set the command line property
 					DesiredCommandLine = null;
@@ -276,8 +270,7 @@ namespace UnrealVS
 					string EditingString = null;
 					if (OleArgs.InValue != null)
 					{
-						object[] InArray = OleArgs.InValue as object[];
-						if (InArray != null && 0 < InArray.Length)
+						if (OleArgs.InValue is object[] InArray && 0 < InArray.Length)
 						{
 							EditingString = InArray.Last() as string;
 						}
@@ -317,8 +310,7 @@ namespace UnrealVS
 
 		private void CommitCommandLineText(string CommandLine)
 		{
-			IVsHierarchy ProjectHierarchy;
-			if (UnrealVSPackage.Instance.SolutionBuildManager.get_StartupProject(out ProjectHierarchy) == VSConstants.S_OK && ProjectHierarchy != null)
+			if (UnrealVSPackage.Instance.SolutionBuildManager.get_StartupProject(out IVsHierarchy ProjectHierarchy) == VSConstants.S_OK && ProjectHierarchy != null)
 			{
 				Project SelectedStartupProject = Utils.HierarchyObjectToProject(ProjectHierarchy);
 				if (SelectedStartupProject != null)
@@ -326,8 +318,7 @@ namespace UnrealVS
 					Configuration SelectedConfiguration = SelectedStartupProject.ConfigurationManager.ActiveConfiguration;
 					if (SelectedConfiguration != null)
 					{
-						IVsBuildPropertyStorage PropertyStorage = ProjectHierarchy as IVsBuildPropertyStorage;
-						if (PropertyStorage != null)
+						if (ProjectHierarchy is IVsBuildPropertyStorage PropertyStorage)
 						{
 							string FullCommandLine = CommandLine;
 
@@ -408,8 +399,7 @@ namespace UnrealVS
 								{
 									// net core project use launchSettings.json to control debug arguments instead
 
-									string activeDebugProfile;
-									if (PropertyStorage.GetPropertyValue("ActiveDebugProfile", ProjectConfigurationName, (uint)_PersistStorageType.PST_USER_FILE, out activeDebugProfile) != VSConstants.S_OK)
+									if (PropertyStorage.GetPropertyValue("ActiveDebugProfile", ProjectConfigurationName, (uint)_PersistStorageType.PST_USER_FILE, out string activeDebugProfile) != VSConstants.S_OK)
 									{
 										activeDebugProfile = null;
 									}
@@ -509,5 +499,32 @@ namespace UnrealVS
 
 		/// used to shield against too many property changes in the future
 		private System.Timers.Timer UpdateCommandLineComboTimer = null;
+
+		// Track IDisposable state
+		private bool IsDisposed;
+
+		protected virtual void Dispose(bool disposing)
+		{
+			if (!IsDisposed)
+			{
+				if (disposing)
+				{
+					// Dispose managed state (managed objects)
+					UpdateCommandLineComboTimer?.Dispose();
+					UpdateCommandLineComboTimer = null;
+				}
+
+				// Free unmanaged resources (unmanaged objects) and override finalizer
+				// Set large fields to null
+				IsDisposed = true;
+			}
+		}
+
+		public void Dispose()
+		{
+			// Do not change this code. Put cleanup code in 'Dispose(bool disposing)' method
+			Dispose(disposing: true);
+			GC.SuppressFinalize(this);
+		}
 	}
 }
diff --git a/Engine/Source/Programs/UnrealVS/CompileSingleFile.cs b/Engine/Source/Programs/UnrealVS/CompileSingleFile.cs
--- a/Engine/Source/Programs/UnrealVS/CompileSingleFile.cs
+++ b/Engine/Source/Programs/UnrealVS/CompileSingleFile.cs
@@ -88,8 +88,7 @@ namespace UnrealVS
 			}
 
 			// Grab the current startup project
-			IVsHierarchy ProjectHierarchy;
-			UnrealVSPackage.Instance.SolutionBuildManager.get_StartupProject(out ProjectHierarchy);
+			UnrealVSPackage.Instance.SolutionBuildManager.get_StartupProject(out IVsHierarchy ProjectHierarchy);
 			if (ProjectHierarchy == null)
 			{
 				Logging.WriteLine("CompileSingleFile: ProjectHierarchy not found");
@@ -101,8 +100,7 @@ namespace UnrealVS
 				Logging.WriteLine("CompileSingleFile: StartupProject not found");
 				return false;
 			}
-			Microsoft.VisualStudio.VCProjectEngine.VCProject VCStartupProject = StartupProject.Object as Microsoft.VisualStudio.VCProjectEngine.VCProject;
-			if (VCStartupProject == null)
+			if (!(StartupProject.Object is Microsoft.VisualStudio.VCProjectEngine.VCProject VCStartupProject))
 			{
 				Logging.WriteLine("CompileSingleFile: VCStartupProject not found");
 				return false;
@@ -161,7 +159,7 @@ namespace UnrealVS
 			DTE.Events.BuildEvents.OnBuildBegin += BuildEvents_OnBuildBegin;
 
 			// Create a delegate for handling output messages
-			DataReceivedEventHandler OutputHandler = (Sender, Args) => { if (Args.Data != null) BuildOutputPane.OutputString($"1>  {Args.Data}{Environment.NewLine}"); };
+			void OutputHandler(object Sender, DataReceivedEventArgs Args) { if (Args.Data != null) BuildOutputPane.OutputString($"1>  {Args.Data}{Environment.NewLine}"); }
 
 			// Get the build command line and escape any environment variables that we use
 			string BuildCommandLine = ActiveNMakeTool.BuildCommandLine;
diff --git a/Engine/Source/Programs/UnrealVS/GenerateProjectFiles.cs b/Engine/Source/Programs/UnrealVS/GenerateProjectFiles.cs
--- a/Engine/Source/Programs/UnrealVS/GenerateProjectFiles.cs
+++ b/Engine/Source/Programs/UnrealVS/GenerateProjectFiles.cs
@@ -1,17 +1,9 @@
 // Copyright Epic Games, Inc. All Rights Reserved.
 
 using System;
-using System.Diagnostics;
 using System.ComponentModel.Design;
+using System.Diagnostics;
 using System.IO;
-using Microsoft.VisualStudio;
-using Microsoft.VisualStudio.Shell;
-using Microsoft.VisualStudio.Shell.Interop;
-using Microsoft.VisualStudio.OLE.Interop;
-using EnvDTE;
-using System.Runtime.InteropServices;
-using System.Collections.Generic;
-using EnvDTE80;
 
 
 namespace UnrealVS
@@ -25,9 +17,9 @@ namespace UnrealVS
 		{
 			// GenerateProjectFilesButton
 			{
-				var CommandID = new CommandID( GuidList.UnrealVSCmdSet, GenerateProjectFilesButtonID );
-				GenerateProjectFilesButtonCommand = new MenuCommand( new EventHandler( GenerateProjectFilesButtonHandler ), CommandID );
-				UnrealVSPackage.Instance.MenuCommandService.AddCommand( GenerateProjectFilesButtonCommand );
+				var CommandID = new CommandID(GuidList.UnrealVSCmdSet, GenerateProjectFilesButtonID);
+				GenerateProjectFilesButtonCommand = new MenuCommand(new EventHandler(GenerateProjectFilesButtonHandler), CommandID);
+				UnrealVSPackage.Instance.MenuCommandService.AddCommand(GenerateProjectFilesButtonCommand);
 			}
 
 
@@ -41,22 +33,22 @@ namespace UnrealVS
 
 
 		/// Called when 'GenerateProjectFiles' button is clicked
-		void GenerateProjectFilesButtonHandler( object Sender, EventArgs Args )
+		void GenerateProjectFilesButtonHandler(object Sender, EventArgs Args)
 		{
 			// Don't allow a regen if we're already generating projects
-			if( GenerateProjectFilesProcess == null )
+			if (GenerateProjectFilesProcess == null)
 			{
 				// Figure out which project to rebuild
 				string BatchFileName = GetBatchFileName();
 
-				if( !String.IsNullOrEmpty( BatchFileName ) )
+				if (!String.IsNullOrEmpty(BatchFileName))
 				{
 					// Make sure the Output window is visible
-					UnrealVSPackage.Instance.DTE.ExecuteCommand( "View.Output" );
+					UnrealVSPackage.Instance.DTE.ExecuteCommand("View.Output");
 
 					// Activate the output pane
 					var Pane = UnrealVSPackage.Instance.GetOutputPane();
-					if( Pane != null )
+					if (Pane != null)
 					{
 						// Clear and activate the output pane.
 						Pane.Clear();
@@ -65,10 +57,10 @@ namespace UnrealVS
 						Pane.Activate();
 					}
 
-					GenerateProjectFilesProcess = UnrealVSPackage.Instance.LaunchProgram(
-                        "cmd.exe", "/C " + GetSafeFilePath(BatchFileName),
+					GenerateProjectFilesProcess = UnrealVSPackage.LaunchProgram(
+						"cmd.exe", "/C " + GetSafeFilePath(BatchFileName),
 						OnGenerateProjectFilesProcessExit,
-						OnOutputFromGenerateProjectFilesProcess );
+						OnOutputFromGenerateProjectFilesProcess);
 				}
 
 				UpdateGenerateProjectFilesButtonState();
@@ -77,7 +69,7 @@ namespace UnrealVS
 
 
 		/// Called when the GenerateProjectFiles process terminates
-		void OnGenerateProjectFilesProcessExit( object Sender, EventArgs Args )
+		void OnGenerateProjectFilesProcessExit(object Sender, EventArgs Args)
 		{
 			GenerateProjectFilesProcess = null;
 
@@ -88,17 +80,17 @@ namespace UnrealVS
 		/// <summary>
 		/// Updates the enabled/disabled state of the 'Generate Project Files' button
 		/// </summary>
-		void UpdateGenerateProjectFilesButtonState( bool ShouldForceDisable = false )
+		void UpdateGenerateProjectFilesButtonState(bool ShouldForceDisable = false)
 		{
 			var CanGenerateProjects = false;
 
-			if( !ShouldForceDisable )
+			if (!ShouldForceDisable)
 			{
 				// Can't launch when a process is already running
-				if( GenerateProjectFilesProcess == null )
+				if (GenerateProjectFilesProcess == null)
 				{
 					string BatchFileName = GetBatchFileName();
-					if ( BatchFileName != null )
+					if (BatchFileName != null)
 					{
 						CanGenerateProjects = true;
 					}
@@ -127,19 +119,19 @@ namespace UnrealVS
 		/// </summary>
 		public void OnSolutionClosing()
 		{
-			UpdateGenerateProjectFilesButtonState( ShouldForceDisable:true );
+			UpdateGenerateProjectFilesButtonState(ShouldForceDisable: true);
 		}
 
 
 		/// <summary>
 		/// Called when data is received from the GenerateProjectFiles process
 		/// </summary>
-		public void OnOutputFromGenerateProjectFilesProcess( object Sender, DataReceivedEventArgs Args )
+		public void OnOutputFromGenerateProjectFilesProcess(object _Sender, DataReceivedEventArgs Args)
 		{
 			var Pane = UnrealVSPackage.Instance.GetOutputPane();
-			if( Pane != null )
+			if (Pane != null)
 			{
-				Pane.OutputString( Args.Data + "\n" );
+				Pane.OutputString(Args.Data + "\n");
 			}
 		}
 
@@ -147,34 +139,34 @@ namespace UnrealVS
 		{
 			// Check to see if we have UE4.sln loaded
 			if (UnrealVSPackage.Instance.IsUESolutionLoaded)
-			{ 
+			{
 				// We expect "GenerateProjectFiles.bat" to live in the same directory as the solution
 				return Path.Combine(Path.GetDirectoryName(UnrealVSPackage.Instance.SolutionFilepath), "GenerateProjectFiles.bat");
 			}
 			return null;
 		}
 
-        /// <summary>
-        /// Wraps a file path in quotes if it contains a space character
-        /// </summary>
-        /// <param name="InPath"></param>
-        /// <returns></returns>
-        private string GetSafeFilePath( string InPath )
-        {
-            string WorkingPath = InPath;
-
-            if (WorkingPath.Contains(" ") && !WorkingPath.Contains("\""))
-            {
-                WorkingPath = "\"" + WorkingPath + "\"";
-            }
-            return WorkingPath;
-        }
+		/// <summary>
+		/// Wraps a file path in quotes if it contains a space character
+		/// </summary>
+		/// <param name="InPath"></param>
+		/// <returns></returns>
+		private string GetSafeFilePath(string InPath)
+		{
+			string WorkingPath = InPath;
+
+			if (WorkingPath.Contains(" ") && !WorkingPath.Contains("\""))
+			{
+				WorkingPath = "\"" + WorkingPath + "\"";
+			}
+			return WorkingPath;
+		}
 
 		/// Active process for the "generate project files" command
 		System.Diagnostics.Process GenerateProjectFilesProcess;
 
 		/// Command for 'Generate Project Files' button
-		MenuCommand GenerateProjectFilesButtonCommand;
+		readonly MenuCommand GenerateProjectFilesButtonCommand;
 	}
 
 }
diff --git a/Engine/Source/Programs/UnrealVS/Guids.cs b/Engine/Source/Programs/UnrealVS/Guids.cs
--- a/Engine/Source/Programs/UnrealVS/Guids.cs
+++ b/Engine/Source/Programs/UnrealVS/Guids.cs
@@ -1,4 +1,4 @@
-﻿// Copyright Epic Games, Inc. All Rights Reserved.
+// Copyright Epic Games, Inc. All Rights Reserved.
 
 // Guids.cs
 // MUST match guids.h
@@ -6,19 +6,19 @@ using System;
 
 namespace UnrealVS
 {
-    static class GuidList
-    {
-        public const string UnrealVSPackageString = "ddbf523f-7eb6-4887-bd51-85a714ff87eb";
-        public const string UnrealVSCmdSetString = "f10be850-eeec-4cef-ac64-1bc2cbe0d447";
-        public const string UnrealVSPackageManString = "683EF1CF-9C97-46D7-BCFA-F12CA6FB77D6";
+	static class GuidList
+	{
+		public const string UnrealVSPackageString = "ddbf523f-7eb6-4887-bd51-85a714ff87eb";
+		public const string UnrealVSCmdSetString = "f10be850-eeec-4cef-ac64-1bc2cbe0d447";
+		public const string UnrealVSPackageManString = "683EF1CF-9C97-46D7-BCFA-F12CA6FB77D6";
 		public const string UnrealVSBatchBuildToolWindowPersistanceString = "9DE31877-5741-4B8B-9F3B-633555D7842B";
 		public const string UnrealVSOptionsString = "FAC770BD-A983-4F40-9364-868D645AF16D";
 
-        // spare generated GUIDs
-        //[Guid("DFF26599-0F46-4A6B-8204-45CC89BAB09D")]
-        //[Guid("E7E7B844-44CF-4DCE-A650-3F97A63C5743")]
+		// spare generated GUIDs
+		//[Guid("DFF26599-0F46-4A6B-8204-45CC89BAB09D")]
+		//[Guid("E7E7B844-44CF-4DCE-A650-3F97A63C5743")]
 
-        public static readonly Guid UnrealVSCmdSet = new Guid(UnrealVSCmdSetString);
+		public static readonly Guid UnrealVSCmdSet = new Guid(UnrealVSCmdSetString);
 
 		public static readonly Guid BatchBuildPaneGuid = new Guid("BBAC0710-4A2E-4BD3-9A39-4683FC6CFE86");
 
@@ -29,5 +29,5 @@ namespace UnrealVS
 		public const string VCProjectKindGuidString = "{8BC9CEB8-8B4A-11D0-8D11-00A0C91BC942}";
 		public const string VJSharpProjectKindGuidString = "{E6FDF86B-F3D1-11D4-8576-0002A516ECE8}";
 		public const string WebProjectKindGuidString = "{E24C65DC-7377-472b-9ABA-BC803B73C61A}";
-    };
+	};
 }
\ No newline at end of file
diff --git a/Engine/Source/Programs/UnrealVS/Logging.cs b/Engine/Source/Programs/UnrealVS/Logging.cs
--- a/Engine/Source/Programs/UnrealVS/Logging.cs
+++ b/Engine/Source/Programs/UnrealVS/Logging.cs
@@ -1,12 +1,8 @@
 // Copyright Epic Games, Inc. All Rights Reserved.
 
 using System;
-using System.Collections.Generic;
 using System.Diagnostics;
 using System.IO;
-using System.Linq;
-using System.Text;
-using System.Threading.Tasks;
 
 namespace UnrealVS
 {
@@ -73,6 +69,10 @@ namespace UnrealVS
 					WriteLine("Visual Studio 13 build");
 #elif VS14
 					WriteLine("Visual Studio 14 build");
+#elif VS15
+					WriteLine("Visual Studio 15 build");
+#elif VS16
+					WriteLine("Visual Studio 16 build");
 #else
 					WriteLine("UNKNOWN Visual Studio build");
 #endif
diff --git a/Engine/Source/Programs/UnrealVS/ProfileManager.cs b/Engine/Source/Programs/UnrealVS/ProfileManager.cs
--- a/Engine/Source/Programs/UnrealVS/ProfileManager.cs
+++ b/Engine/Source/Programs/UnrealVS/ProfileManager.cs
@@ -1,44 +1,41 @@
-﻿// Copyright Epic Games, Inc. All Rights Reserved.
+// Copyright Epic Games, Inc. All Rights Reserved.
 
+using Microsoft.VisualStudio.Shell;
+using Microsoft.VisualStudio.Shell.Interop;
 using System;
-using System.Collections.Generic;
 using System.ComponentModel;
-using System.Linq;
-using System.Text;
 using System.Runtime.InteropServices;
-using Microsoft.VisualStudio.Shell;
-using Microsoft.VisualStudio.Shell.Interop;
 
 namespace UnrealVS
 {
-    // GUID for this class.
-    [Guid(GuidList.UnrealVSPackageManString)]
+	// GUID for this class.
+	[Guid(GuidList.UnrealVSPackageManString)]
 
-    class ProfileManager : Component, IProfileManager
-    {
-        public void SaveSettingsToXml(IVsSettingsWriter writer)
-        {
-            throw new NotImplementedException();
-        }
+	class ProfileManager : Component, IProfileManager
+	{
+		public void SaveSettingsToXml(IVsSettingsWriter writer)
+		{
+			throw new NotImplementedException();
+		}
 
-        public void LoadSettingsFromXml(IVsSettingsReader reader)
-        {
-            throw new NotImplementedException();
-        }
+		public void LoadSettingsFromXml(IVsSettingsReader reader)
+		{
+			throw new NotImplementedException();
+		}
 
-        public void SaveSettingsToStorage()
-        {
-            throw new NotImplementedException();
-        }
+		public void SaveSettingsToStorage()
+		{
+			throw new NotImplementedException();
+		}
 
-        public void LoadSettingsFromStorage()
-        {
-            throw new NotImplementedException();
-        }
+		public void LoadSettingsFromStorage()
+		{
+			throw new NotImplementedException();
+		}
 
-        public void ResetSettings()
-        {
-            throw new NotImplementedException();
-        }
-    }
+		public void ResetSettings()
+		{
+			throw new NotImplementedException();
+		}
+	}
 }
diff --git a/Engine/Source/Programs/UnrealVS/Properties/AssemblyInfo.cs b/Engine/Source/Programs/UnrealVS/Properties/AssemblyInfo.cs
--- a/Engine/Source/Programs/UnrealVS/Properties/AssemblyInfo.cs
+++ b/Engine/Source/Programs/UnrealVS/Properties/AssemblyInfo.cs
@@ -3,7 +3,6 @@
 using System;
 using System.Reflection;
 using System.Resources;
-using System.Runtime.CompilerServices;
 using System.Runtime.InteropServices;
 
 // General Information about an assembly is controlled through the following 
@@ -16,8 +15,8 @@ using System.Runtime.InteropServices;
 [assembly: AssemblyProduct("UnrealVS")]
 [assembly: AssemblyCopyright("Copyright Epic Games, Inc.")]
 [assembly: AssemblyTrademark("")]
-[assembly: AssemblyCulture("")]   
-[assembly: ComVisible(false)]     
+[assembly: AssemblyCulture("")]
+[assembly: ComVisible(false)]
 [assembly: CLSCompliant(false)]
 [assembly: NeutralResourcesLanguage("en-US")]
 
diff --git a/Engine/Source/Programs/UnrealVS/ProvideSolutionProperties.cs b/Engine/Source/Programs/UnrealVS/ProvideSolutionProperties.cs
--- a/Engine/Source/Programs/UnrealVS/ProvideSolutionProperties.cs
+++ b/Engine/Source/Programs/UnrealVS/ProvideSolutionProperties.cs
@@ -1,8 +1,8 @@
 // Copyright Epic Games, Inc. All Rights Reserved.
 
+using Microsoft.VisualStudio.Shell;
 using System;
 using System.Globalization;
-using Microsoft.VisualStudio.Shell;
 
 namespace UnrealVS
 {
@@ -12,7 +12,7 @@ namespace UnrealVS
 	[AttributeUsage(AttributeTargets.Class, AllowMultiple = true, Inherited = true)]
 	internal sealed class ProvideSolutionPropertiesAttribute : RegistrationAttribute
 	{
-		private string _propName;
+		private readonly string _propName;
 
 		public ProvideSolutionPropertiesAttribute(string propName)
 		{
diff --git a/Engine/Source/Programs/UnrealVS/QuickBuild.cs b/Engine/Source/Programs/UnrealVS/QuickBuild.cs
--- a/Engine/Source/Programs/UnrealVS/QuickBuild.cs
+++ b/Engine/Source/Programs/UnrealVS/QuickBuild.cs
@@ -10,9 +10,6 @@ using System.Collections.Generic;
 using System.ComponentModel.Design;
 using System.Linq;
 using System.Runtime.InteropServices;
-using System.Text;
-using System.Threading.Tasks;
-using System.Windows.Forms.VisualStyles;
 
 namespace UnrealVS
 {
@@ -23,7 +20,7 @@ namespace UnrealVS
 	class QuickBuildMenuCommand : OleMenuCommand
 	{
 		/// The function used to select the range of commands handled by this object
-		private Predicate<int> OnMatchItem;
+		private readonly Predicate<int> OnMatchItem;
 
 		/// <summary>
 		/// Create a QuickBuildMenuCommand with a delegate for determining which command numbers it handles
@@ -32,12 +29,7 @@ namespace UnrealVS
 		public QuickBuildMenuCommand(CommandID InDynamicStartCommandId, Predicate<int> InOnMatchItem, EventHandler InOnInvokedDynamicItem, EventHandler InOnBeforeQueryStatusDynamicItem)
 			: base(InOnInvokedDynamicItem, null /*changeHandler*/, InOnBeforeQueryStatusDynamicItem, InDynamicStartCommandId)
 		{
-			if (InOnMatchItem == null)
-			{
-				throw new ArgumentNullException("InOnMatchItem");
-			}
-
-			OnMatchItem = InOnMatchItem;
+			OnMatchItem = InOnMatchItem ?? throw new ArgumentNullException("InOnMatchItem");
 		}
 
 		/// <summary>
@@ -115,14 +107,10 @@ namespace UnrealVS
 			var MenuCommand = (QuickBuildMenuCommand)sender;
 
 			// Get the project clicked in the solution explorer by accessing the current selection and converting to a Project if possible.
-			IntPtr HierarchyPtr, SelectionContainerPtr;
-			uint ProjectItemId;
-			IVsMultiItemSelect MultiItemSelect;
-			UnrealVSPackage.Instance.SelectionManager.GetCurrentSelection(out HierarchyPtr, out ProjectItemId, out MultiItemSelect, out SelectionContainerPtr);
+			UnrealVSPackage.Instance.SelectionManager.GetCurrentSelection(out IntPtr HierarchyPtr, out uint ProjectItemId, out IVsMultiItemSelect MultiItemSelect, out IntPtr SelectionContainerPtr);
 			if (HierarchyPtr == null) return;
 
-			IVsHierarchy SelectedHierarchy = Marshal.GetTypedObjectForIUnknown(HierarchyPtr, typeof(IVsHierarchy)) as IVsHierarchy;
-			if (SelectedHierarchy == null) return;
+			if (!(Marshal.GetTypedObjectForIUnknown(HierarchyPtr, typeof(IVsHierarchy)) is IVsHierarchy SelectedHierarchy)) return;
 
 			var SelectedProject = Utils.HierarchyObjectToProject(SelectedHierarchy);
 			if (SelectedProject == null) return;
@@ -173,7 +161,7 @@ namespace UnrealVS
 
 		/** constants */
 
-		private const int ProjectQuickBuildMenuID = 0x1410;		// must match the values in the vsct file
+		private const int ProjectQuickBuildMenuID = 0x1410;     // must match the values in the vsct file
 
 		/** fields */
 
@@ -183,7 +171,7 @@ namespace UnrealVS
 
 		/// Hide/shows the while Quick Build menu tree
 		private static bool bIsActive = false;
-		
+
 		/// List of submenus and their details - must match the values in the vsct file
 		private readonly SubMenu[] SubMenus = new[]
 		{
@@ -248,7 +236,7 @@ namespace UnrealVS
 			QuickBuildCommand = new OleMenuCommand(null, null, OnQuickBuildQuery, new CommandID(GuidList.UnrealVSCmdSet, ProjectQuickBuildMenuID));
 			QuickBuildCommand.BeforeQueryStatus += OnQuickBuildQuery;
 			UnrealVSPackage.Instance.MenuCommandService.AddCommand(QuickBuildCommand);
-			
+
 			// configuration sub-menus
 			foreach (var SubMenu in SubMenus)
 			{
@@ -275,8 +263,7 @@ namespace UnrealVS
 		/// </summary>
 		private void UpdateActiveState()
 		{
-			int bIsBuilding;
-			UnrealVSPackage.Instance.SelectionManager.IsCmdUIContextActive(SolutionBuildingUIContextCookie, out bIsBuilding);
+			UnrealVSPackage.Instance.SelectionManager.IsCmdUIContextActive(SolutionBuildingUIContextCookie, out int bIsBuilding);
 			bIsActive = bIsBuilding == 0;
 		}
 
diff --git a/Engine/Source/Programs/UnrealVS/StartupProjectSelector.cs b/Engine/Source/Programs/UnrealVS/StartupProjectSelector.cs
--- a/Engine/Source/Programs/UnrealVS/StartupProjectSelector.cs
+++ b/Engine/Source/Programs/UnrealVS/StartupProjectSelector.cs
@@ -1,14 +1,13 @@
-﻿// Copyright Epic Games, Inc. All Rights Reserved.
+// Copyright Epic Games, Inc. All Rights Reserved.
 
-using System;
-using System.ComponentModel.Design;
+using EnvDTE;
 using Microsoft.VisualStudio.Shell;
 using Microsoft.VisualStudio.Shell.Interop;
-using EnvDTE;
-using System.Runtime.InteropServices;
+using System;
 using System.Collections.Generic;
+using System.ComponentModel.Design;
 using System.Linq;
-using System.Diagnostics;
+using System.Runtime.InteropServices;
 
 namespace UnrealVS
 {
@@ -32,35 +31,35 @@ namespace UnrealVS
 
 			public string Name
 			{
-				get { return (Project != null ? Project.Name : "<invalid>" ); }
+				get { return (Project != null ? Project.Name : "<invalid>"); }
 			}
 
 			public override string ToString()
 			{
-				return (Project != null ? Project.Name : "<invalid>" );
+				return (Project != null ? Project.Name : "<invalid>");
 			}
 		}
 
-        public StartupProjectSelector()
+		public StartupProjectSelector()
 		{
 			_bIsSolutionOpened = UnrealVSPackage.Instance.DTE.Solution.IsOpen;
 
 			// Create the handlers for our commands
 			{
 				// StartupProjectCombo
-				var StartupProjectComboCommandId = new CommandID( GuidList.UnrealVSCmdSet, StartupProjectComboId );
+				var StartupProjectComboCommandId = new CommandID(GuidList.UnrealVSCmdSet, StartupProjectComboId);
 				var StartupProjectComboCommand = new OleMenuCommand(StartupProjectComboHandler, StartupProjectComboCommandId);
 				StartupProjectComboCommand.BeforeQueryStatus += (sender, args) => StartupProjectComboCommand.Enabled = _bIsSolutionOpened;
-				UnrealVSPackage.Instance.MenuCommandService.AddCommand( StartupProjectComboCommand );
+				UnrealVSPackage.Instance.MenuCommandService.AddCommand(StartupProjectComboCommand);
 
 				// StartupProjectComboList
-				var StartupProjectComboListCommandId = new CommandID( GuidList.UnrealVSCmdSet, StartupProjectComboListId );
+				var StartupProjectComboListCommandId = new CommandID(GuidList.UnrealVSCmdSet, StartupProjectComboListId);
 				var StartupProjectComboListCommand = new OleMenuCommand(StartupProjectComboListHandler, StartupProjectComboListCommandId);
 				StartupProjectComboListCommand.BeforeQueryStatus += (sender, args) => StartupProjectComboListCommand.Enabled = _bIsSolutionOpened;
-				UnrealVSPackage.Instance.MenuCommandService.AddCommand( StartupProjectComboListCommand );
+				UnrealVSPackage.Instance.MenuCommandService.AddCommand(StartupProjectComboListCommand);
 
 				StartupProjectComboCommands = new[] { StartupProjectComboCommand, StartupProjectComboListCommand };
-			} 
+			}
 
 			// Register for events that we care about
 			UnrealVSPackage.Instance.OnStartupProjectChanged += OnStartupProjectChanged;
@@ -81,8 +80,8 @@ namespace UnrealVS
 					_CachedStartupProjects.Clear();
 				};
 			UnrealVSPackage.Instance.OnProjectOpened +=
-				delegate(Project OpenedProject)
-				{		
+				delegate (Project OpenedProject)
+				{
 					if (_bIsSolutionOpened)
 					{
 						Logging.WriteLine("Opened project node " + OpenedProject.Name);
@@ -94,7 +93,7 @@ namespace UnrealVS
 					}
 				};
 			UnrealVSPackage.Instance.OnProjectClosed +=
-				delegate(Project ClosedProject)
+				delegate (Project ClosedProject)
 				{
 					Logging.WriteLine("Closed project node " + ClosedProject.Name);
 					RemoveFromStartupProjectList(ClosedProject);
@@ -145,18 +144,15 @@ namespace UnrealVS
 
 			SortCachedStartupProjects();
 
-			if (StartupProjectListChanged != null)
-			{
-				StartupProjectListChanged(this, new StartupProjectListChangedEventArgs {StartupProjects = StartupProjectList});
-			}
+			StartupProjectListChanged?.Invoke(this, new StartupProjectListChangedEventArgs { StartupProjects = StartupProjectList });
 		}
 
 		private void SortCachedStartupProjects()
 		{
-            // Sort projects by game y/n then alphabetically
-            _CachedStartupProjects = (_CachedStartupProjects.OrderBy(ProjectRef => ProjectRef.Name == "UE4" ? 0 : 1)
-                .ThenBy(ProjectRef => Utils.IsGameProject(ProjectRef.Project) ? 0 : 1)
-                .ThenBy(ProjectRef => ProjectRef.Name)).ToList();
+			// Sort projects by game y/n then alphabetically
+			_CachedStartupProjects = (_CachedStartupProjects.OrderBy(ProjectRef => ProjectRef.Name == "UE4" ? 0 : 1)
+				.ThenBy(ProjectRef => Utils.IsGameProject(ProjectRef.Project) ? 0 : 1)
+				.ThenBy(ProjectRef => ProjectRef.Name)).ToList();
 		}
 
 		private bool ProjectPredicate(Project Project)
@@ -193,18 +189,15 @@ namespace UnrealVS
 		{
 			_CachedStartupProjects.RemoveAll(ProjRef => ProjRef.Project == Project);
 
-			if (StartupProjectListChanged != null)
-			{
-				StartupProjectListChanged(this, new StartupProjectListChangedEventArgs { StartupProjects = StartupProjectList });
-			}
-		}	
+			StartupProjectListChanged?.Invoke(this, new StartupProjectListChangedEventArgs { StartupProjects = StartupProjectList });
+		}
 
 		/// <summary>
 		/// Updates the combo box after the project loaded state has changed
 		/// </summary>
 		private void UpdateStartupProjectCombo()
 		{
-			foreach(var Command in StartupProjectComboCommands)
+			foreach (var Command in StartupProjectComboCommands)
 			{
 				Command.Enabled = _bIsSolutionOpened;
 			}
@@ -214,12 +207,12 @@ namespace UnrealVS
 		/// <summary>
 		/// Called when the startup project has changed to a new project.  We'll refresh our interface.
 		/// </summary>
-		public void OnStartupProjectChanged( Project NewStartupProject )
+		public void OnStartupProjectChanged(Project NewStartupProject)
 		{
 			UpdateStartupProjectCombo();
 		}
 
-	
+
 		/// <summary>
 		/// Returns a string to display in the startup project combo box, based on project state
 		/// </summary>
@@ -229,9 +222,8 @@ namespace UnrealVS
 			string Text = "";
 
 			// Switch to this project!
-			IVsHierarchy ProjectHierarchy;
-			UnrealVSPackage.Instance.SolutionBuildManager.get_StartupProject( out ProjectHierarchy );
-			if( ProjectHierarchy != null )
+			UnrealVSPackage.Instance.SolutionBuildManager.get_StartupProject(out IVsHierarchy ProjectHierarchy);
+			if (ProjectHierarchy != null)
 			{
 				var Project = Utils.HierarchyObjectToProject(ProjectHierarchy);
 				if (Project != null)
@@ -245,19 +237,19 @@ namespace UnrealVS
 
 
 		/// Called by combo control to query the text to display or to apply newly-entered text
-		void StartupProjectComboHandler( object Sender, EventArgs Args )
+		void StartupProjectComboHandler(object Sender, EventArgs Args)
 		{
 			var OleArgs = (OleMenuCmdEventArgs)Args;
-			
+
 			// If OutValue is non-zero, then Visual Studio is querying the current combo value to display
-			if( OleArgs.OutValue != IntPtr.Zero )
+			if (OleArgs.OutValue != IntPtr.Zero)
 			{
 				string ValueString = MakeStartupProjectComboText();
-				Marshal.GetNativeVariantForObject( ValueString, OleArgs.OutValue );
+				Marshal.GetNativeVariantForObject(ValueString, OleArgs.OutValue);
 			}
 
 			// Has a new input value been entered?
-			else if( OleArgs.InValue != null )
+			else if (OleArgs.InValue != null)
 			{
 				var InputString = OleArgs.InValue.ToString();
 
@@ -276,12 +268,12 @@ namespace UnrealVS
 					{
 						UnrealVSPackage.Instance.DTE.ExecuteCommand("Project.SetasStartUpProject");
 					}
-				}					
+				}
 			}
 		}
 
 		/// Called by combo control to populate the drop-down list
-		void StartupProjectComboListHandler( object Sender, EventArgs Args )
+		void StartupProjectComboListHandler(object Sender, EventArgs Args)
 		{
 			var OleArgs = (OleMenuCmdEventArgs)Args;
 
diff --git a/Engine/Source/Programs/UnrealVS/UnrealVSOptions.cs b/Engine/Source/Programs/UnrealVS/UnrealVSOptions.cs
--- a/Engine/Source/Programs/UnrealVS/UnrealVSOptions.cs
+++ b/Engine/Source/Programs/UnrealVS/UnrealVSOptions.cs
@@ -1,9 +1,9 @@
-﻿// Copyright Epic Games, Inc. All Rights Reserved.
+// Copyright Epic Games, Inc. All Rights Reserved.
 
+using Microsoft.VisualStudio.Shell;
 using System;
 using System.ComponentModel;
 using System.Runtime.InteropServices;
-using Microsoft.VisualStudio.Shell;
 
 namespace UnrealVS
 {
diff --git a/Engine/Source/Programs/UnrealVS/UnrealVSPackage.cs b/Engine/Source/Programs/UnrealVS/UnrealVSPackage.cs
--- a/Engine/Source/Programs/UnrealVS/UnrealVSPackage.cs
+++ b/Engine/Source/Programs/UnrealVS/UnrealVSPackage.cs
@@ -1,22 +1,19 @@
 // Copyright Epic Games, Inc. All Rights Reserved.
 
+using EnvDTE;
+using EnvDTE80;
+using Microsoft.VisualStudio;
+using Microsoft.VisualStudio.OLE.Interop;
+using Microsoft.VisualStudio.Shell;
+using Microsoft.VisualStudio.Shell.Interop;
 using System;
 using System.Collections.Generic;
+using System.ComponentModel.Design;
 using System.Diagnostics;
 using System.IO;
 using System.Linq;
 using System.Runtime.InteropServices;
-using System.ComponentModel.Design;
 using System.Threading;
-using System.Windows;
-using System.Windows.Forms;
-using Microsoft.VisualStudio;
-using Microsoft.VisualStudio.Shell.Interop;
-using Microsoft.VisualStudio.OLE.Interop;
-using Microsoft.VisualStudio.Shell;
-using EnvDTE;
-using EnvDTE80;
-using MessageBox = System.Windows.Forms.MessageBox;
 using Thread = System.Threading.Thread;
 
 namespace UnrealVS
@@ -32,14 +29,14 @@ namespace UnrealVS
 	}
 
 	// This attribute tells the PkgDef creation utility (CreatePkgDef.exe) that this class is a VS package.
-	[PackageRegistration( UseManagedResourcesOnly = true, AllowsBackgroundLoading = true )]
+	[PackageRegistration(UseManagedResourcesOnly = true, AllowsBackgroundLoading = true)]
 
 	// This attribute is used to register the informations needed to show the this package in the Help/About dialog of Visual Studio.
 	[InstalledProductRegistration("#110", "#112", VersionString, IconResourceID = 400)]
 
 	// Adds data for our interface elements defined in UnrealVS.vsct.  The name "Menus.ctmenu" here is arbitrary, but must
 	// match the ItemGroup inside the UnrealVS.csproj MSBuild file (hand-typed.)
-	[ProvideMenuResource( "Menus.ctmenu", 1 )]
+	[ProvideMenuResource("Menus.ctmenu", 1)]
 
 	// Force the package to load whenever a solution exists
 	//[ProvideAutoLoad(VSConstants.UICONTEXT.SolutionExistsAndFullyLoaded_string, PackageAutoLoadFlags.BackgroundLoad)]
@@ -49,7 +46,7 @@ namespace UnrealVS
 	//[ProvideProfile(typeof (ProfileManager), "UnrealVS", "UnrealVSPackage", 110, 113, false)]
 
 	// GUID for this class.  This needs to match in quite a few places in the project.
-	[Guid( GuidList.UnrealVSPackageString )]
+	[Guid(GuidList.UnrealVSPackageString)]
 
 	// This attribute registers a tool window exposed by this package.
 	[ProvideToolWindow(typeof(BatchBuilderToolWindow))]
@@ -62,10 +59,10 @@ namespace UnrealVS
 	/// with Visual Studio shell and serves as the entry point into our extension
 	/// </summary>
 	public sealed class UnrealVSPackage :
-		AsyncPackage,			// We inherit from AsyncPackage which allows us to become a "plugin" within the Visual Studio shell
-		IVsSolutionEvents,		// This interface allows us to register to be notified of events such as opening a project
+		AsyncPackage,           // We inherit from AsyncPackage which allows us to become a "plugin" within the Visual Studio shell
+		IVsSolutionEvents,      // This interface allows us to register to be notified of events such as opening a project
 		IVsUpdateSolutionEvents,// Allows us to register to be notified of events such as active config changes
-		IVsSelectionEvents,		// Allows us to be notified when the startup project has changed to a different project
+		IVsSelectionEvents,     // Allows us to be notified when the startup project has changed to a different project
 		IVsHierarchyEvents,     // Allows us to be notified when a hierarchy (the startup project) has had properties changed
 		IVsPersistSolutionProps, // Allows us to read props added to the solution to determine if the solution is a unreal solution
 		IDisposable
@@ -83,54 +80,43 @@ namespace UnrealVS
 
 		/// Called when a new startup project is set in Visual Studio
 		public delegate void OnStartupProjectChangedDelegate(Project NewStartupProject);
-		[System.Diagnostics.CodeAnalysis.SuppressMessage( "Microsoft.Design", "CA1009" )]
 		public event OnStartupProjectChangedDelegate OnStartupProjectChanged;
 
 		/// Called when a project is opened or created in Visual Studio
 		public delegate void OnProjectOpenedDelegate(Project OpenedProject);
-		[System.Diagnostics.CodeAnalysis.SuppressMessage( "Microsoft.Design", "CA1009" )]
 		public event OnProjectOpenedDelegate OnProjectOpened;
 
 		/// Called right before a project is closed
 		public delegate void OnProjectClosedDelegate(Project ClosedProject);
-		[System.Diagnostics.CodeAnalysis.SuppressMessage( "Microsoft.Design", "CA1009" )]
 		public event OnProjectClosedDelegate OnProjectClosed;
 
 		/// Called when the startup project is edited in Visual Studio
 		public delegate void OnStartupProjectPropertyChangedDelegate(UInt32 itemid, Int32 propid, UInt32 flags);
-		[System.Diagnostics.CodeAnalysis.SuppressMessage("Microsoft.Design", "CA1009")]
 		public event OnStartupProjectPropertyChangedDelegate OnStartupProjectPropertyChanged;
 
 		/// Called right after a solution is opened
 		public delegate void OnSolutionOpenedDelegate();
-		[System.Diagnostics.CodeAnalysis.SuppressMessage( "Microsoft.Design", "CA1009" )]
 		public event OnSolutionOpenedDelegate OnSolutionOpened;
 
 		/// Called right before/after a solution is closed
 		public delegate void OnSolutionClosedDelegate();
-		[System.Diagnostics.CodeAnalysis.SuppressMessage( "Microsoft.Design", "CA1009" )]
 		public event OnSolutionClosedDelegate OnSolutionClosing;
-		[System.Diagnostics.CodeAnalysis.SuppressMessage( "Microsoft.Design", "CA1009" )]
 		public event OnSolutionClosedDelegate OnSolutionClosed;
 
 		/// Called when the active project config changes for any project
 		public delegate void OnStartupProjectConfigChangedDelegate(Project Project);
-		[System.Diagnostics.CodeAnalysis.SuppressMessage("Microsoft.Design", "CA1009")]
 		public event OnStartupProjectConfigChangedDelegate OnStartupProjectConfigChanged;
 
 		/// Called when a build/update action begins
 		public delegate void OnBuildBeginDelegate(out int Cancel);
-		[System.Diagnostics.CodeAnalysis.SuppressMessage("Microsoft.Design", "CA1009")]
 		public event OnBuildBeginDelegate OnBuildBegin;
 
 		/// Called when a build/update action completes
 		public delegate void OnBuildDoneDelegate(bool bSucceeded, bool bModified, bool bWasCancelled);
-		[System.Diagnostics.CodeAnalysis.SuppressMessage( "Microsoft.Design", "CA1009" )]
 		public event OnBuildDoneDelegate OnBuildDone;
 
 		/// Called when the UIContext changes
 		public delegate void OnUIContextChangedDelegate(uint CmdUICookie, bool bActive);
-		[System.Diagnostics.CodeAnalysis.SuppressMessage("Microsoft.Design", "CA1009")]
 		public event OnUIContextChangedDelegate OnUIContextChanged;
 
 		/** Public Fields & Properties */
@@ -142,7 +128,7 @@ namespace UnrealVS
 		}
 
 		/// Visual Studio menu command service
-		public IMenuCommandService MenuCommandService	{ get; private set; }
+		public IMenuCommandService MenuCommandService { get; private set; }
 
 		/// Visual Studio solution build manager interface.  This is used to change the active startup
 		/// Project, among other things.  We expose public access to the solution build manager through
@@ -208,8 +194,8 @@ namespace UnrealVS
 			await JoinableTaskFactory.SwitchToMainThreadAsync(cancellationToken);
 
 			// Grab the MenuCommandService
-			MenuCommandService = await GetServiceAsync( typeof( IMenuCommandService ) ) as OleMenuCommandService;
-			
+			MenuCommandService = await GetServiceAsync(typeof(IMenuCommandService)) as OleMenuCommandService;
+
 			// Get access to Visual Studio's DTE object.  This object has various hooks into the Visual Studio
 			// shell that are useful for writing extensions.
 			DTE = await GetServiceAsync(typeof(DTE)) as DTE;
@@ -266,9 +252,11 @@ namespace UnrealVS
 		{
 			// Create a "ticker" on a background thread that ticks the package on the UI thread
 			Interlocked.Exchange(ref bCancelTicker, 0);
-			Ticker = new Thread(TickAsyncMain);
-			Ticker.Priority = ThreadPriority.Lowest;
-			Ticker.Start();			
+			Ticker = new Thread(TickAsyncMain)
+			{
+				Priority = ThreadPriority.Lowest
+			};
+			Ticker.Start();
 		}
 
 		private void StopTicker()
@@ -288,10 +276,10 @@ namespace UnrealVS
 			{
 				while (true)
 				{
-					if (bCancelTicker != 0) return; 
+					if (bCancelTicker != 0) return;
 					ThreadHelper.Generic.BeginInvoke(Tick);
 
-					if (bCancelTicker != 0) return; 
+					if (bCancelTicker != 0) return;
 					Thread.Sleep(TickPeriod);
 				}
 			}
@@ -317,7 +305,7 @@ namespace UnrealVS
 		}
 
 		/// IDispose pattern lets us clean up our stuff!
-		protected override void Dispose( bool disposing )
+		protected override void Dispose(bool disposing)
 		{
 			if (Ticker != null && Ticker.IsAlive)
 			{
@@ -334,30 +322,29 @@ namespace UnrealVS
 			// Clean up singleton instance
 			PrivateInstance = null;
 
-			CommandLineEditor = null;
 			StartupProjectSelector = null;
 			BatchBuilder = null;
 			QuickBuilder = null;
 
-			if(CompileSingleFile != null)
-			{
-				CompileSingleFile.Dispose();
-				CompileSingleFile = null;
-			}
-			
+			CompileSingleFile?.Dispose();
+			CompileSingleFile = null;
+
+			CommandLineEditor?.Dispose();
+			CommandLineEditor = null;
+
 			// No longer want solution events
-			if( SolutionEventsHandle != 0 )
+			if (SolutionEventsHandle != 0)
 			{
-				SolutionManager.UnadviseSolutionEvents( SolutionEventsHandle );
+				SolutionManager.UnadviseSolutionEvents(SolutionEventsHandle);
 				SolutionEventsHandle = 0;
 			}
 			SolutionManager = null;
 
 			// No longer want selection events
-			if( SelectionEventsHandle != 0 )
+			if (SelectionEventsHandle != 0)
 			{
-				SelectionManager.UnadviseSelectionEvents( SelectionEventsHandle );
-				SelectionEventsHandle = 0;	
+				SelectionManager.UnadviseSelectionEvents(SelectionEventsHandle);
+				SelectionEventsHandle = 0;
 			}
 			SelectionManager = null;
 
@@ -401,7 +388,7 @@ namespace UnrealVS
 		/// The package's options page
 		public UnrealVsOptions OptionsPage
 		{
-			get { return (UnrealVsOptions) GetDialogPage(typeof (UnrealVsOptions)); }
+			get { return (UnrealVsOptions)GetDialogPage(typeof(UnrealVsOptions)); }
 		}
 
 		/// <summary>
@@ -412,14 +399,16 @@ namespace UnrealVS
 		/// <param name="OnExit">Optional callback whent he program exits</param>
 		/// <param name="OutputHandler">If supplied, std-out and std-error will be redirected to this function and no shell window will be created</param>
 		/// <returns>The newly-created process, or null if it wasn't able to start.  Exceptions are swallowed, but a debug output string is emitted on errors</returns>
-		public System.Diagnostics.Process LaunchProgram( string ProgramFile, string Arguments, EventHandler OnExit = null, DataReceivedEventHandler OutputHandler = null, bool bWaitForCompletion=false )
+		public static System.Diagnostics.Process LaunchProgram(string ProgramFile, string Arguments, EventHandler OnExit = null, DataReceivedEventHandler OutputHandler = null, bool bWaitForCompletion = false)
 		{
 			// Create the action's process.
-			ProcessStartInfo ActionStartInfo = new ProcessStartInfo();
-			ActionStartInfo.FileName = ProgramFile;
-			ActionStartInfo.Arguments = Arguments;
+			ProcessStartInfo ActionStartInfo = new ProcessStartInfo
+			{
+				FileName = ProgramFile,
+				Arguments = Arguments
+			};
 
-			if( OutputHandler != null )
+			if (OutputHandler != null)
 			{
 				ActionStartInfo.RedirectStandardInput = true;
 				ActionStartInfo.RedirectStandardOutput = true;
@@ -440,16 +429,18 @@ namespace UnrealVS
 
 			try
 			{
-				ActionProcess = new System.Diagnostics.Process();
-				ActionProcess.StartInfo = ActionStartInfo;
+				ActionProcess = new System.Diagnostics.Process
+				{
+					StartInfo = ActionStartInfo
+				};
 
-				if( OnExit != null )
+				if (OnExit != null)
 				{
 					ActionProcess.EnableRaisingEvents = true;
 					ActionProcess.Exited += OnExit;
 				}
 
-				if( ActionStartInfo.RedirectStandardOutput )
+				if (ActionStartInfo.RedirectStandardOutput)
 				{
 					ActionProcess.EnableRaisingEvents = true;
 					ActionProcess.OutputDataReceived += OutputHandler;
@@ -459,7 +450,7 @@ namespace UnrealVS
 				// Launch the program
 				ActionProcess.Start();
 
-				if( ActionStartInfo.RedirectStandardOutput )
+				if (ActionStartInfo.RedirectStandardOutput)
 				{
 					ActionProcess.BeginOutputReadLine();
 					ActionProcess.BeginErrorReadLine();
@@ -473,7 +464,7 @@ namespace UnrealVS
 					}
 				}
 			}
-			catch( Exception Ex)
+			catch (Exception Ex)
 			{
 				// Couldn't launch program
 				Logging.WriteLine("Couldn't launch program: " + ActionStartInfo.FileName);
@@ -566,91 +557,82 @@ namespace UnrealVS
 		/// IVsSolutionEvents implementation
 		///
 
-		int IVsSolutionEvents.OnAfterCloseSolution( object pUnkReserved )
+		int IVsSolutionEvents.OnAfterCloseSolution(object pUnkReserved)
 		{
 			UpdateUnrealLoadedStatus();
 
-			if (OnSolutionClosed != null)
-			{
-				OnSolutionClosed();
-			}
+			OnSolutionClosed?.Invoke();
 
 			return VSConstants.S_OK;
 		}
 
-		int IVsSolutionEvents.OnAfterLoadProject( IVsHierarchy pStubHierarchy, IVsHierarchy pRealHierarchy )
+		int IVsSolutionEvents.OnAfterLoadProject(IVsHierarchy pStubHierarchy, IVsHierarchy pRealHierarchy)
 		{
 			return VSConstants.S_OK;
 		}
 
-		int IVsSolutionEvents.OnAfterOpenProject( IVsHierarchy pHierarchy, int fAdded )
+		int IVsSolutionEvents.OnAfterOpenProject(IVsHierarchy pHierarchy, int fAdded)
 		{
 			// This function is called after a Visual Studio project is opened (or a new project is created.)
 
 			// Get the actual Project object from the IVsHierarchy object that was supplied
-			var OpenedProject = Utils.HierarchyObjectToProject( pHierarchy );
+			var OpenedProject = Utils.HierarchyObjectToProject(pHierarchy);
 			Utils.OnProjectListChanged();
 			if (OpenedProject != null && OnProjectOpened != null)
 			{
 				LoadedProjectPaths.Add(OpenedProject.FullName);
-                OnProjectOpened(OpenedProject);
-            }
+				OnProjectOpened(OpenedProject);
+			}
 
-		    return VSConstants.S_OK;
+			return VSConstants.S_OK;
 		}
 
-		int IVsSolutionEvents.OnAfterOpenSolution( object pUnkReserved, int fNewSolution )
+		int IVsSolutionEvents.OnAfterOpenSolution(object pUnkReserved, int fNewSolution)
 		{
 			UpdateUnrealLoadedStatus();
 
 			StartTicker();
 
-            if (OnSolutionOpened != null)
-            {
-                OnSolutionOpened();
-            }
+			OnSolutionOpened?.Invoke();
 
-		    return VSConstants.S_OK;
+			return VSConstants.S_OK;
 		}
 
-		int IVsSolutionEvents.OnBeforeCloseProject( IVsHierarchy pHierarchy, int fRemoved )
+		int IVsSolutionEvents.OnBeforeCloseProject(IVsHierarchy pHierarchy, int fRemoved)
 		{
 			// This function is called after a Visual Studio project is closed
 
 			// Get the actual Project object from the IVsHierarchy object that was supplied
-			var ClosedProject = Utils.HierarchyObjectToProject( pHierarchy );
+			var ClosedProject = Utils.HierarchyObjectToProject(pHierarchy);
 			if (ClosedProject != null && OnProjectClosed != null)
-            {
+			{
 				LoadedProjectPaths.Remove(ClosedProject.FullName);
-                OnProjectClosed(ClosedProject);
-            }
+				OnProjectClosed(ClosedProject);
+			}
 
 			return VSConstants.S_OK;
 		}
 
-		int IVsSolutionEvents.OnBeforeCloseSolution( object pUnkReserved )
+		int IVsSolutionEvents.OnBeforeCloseSolution(object pUnkReserved)
 		{
 			StopTicker();
 
-            if (OnSolutionClosing != null)
-            {
-                OnSolutionClosing();
-            }
+			OnSolutionClosing?.Invoke();
 
 			return VSConstants.S_OK;
 		}
 
-		int IVsSolutionEvents.OnBeforeUnloadProject( IVsHierarchy pRealHierarchy, IVsHierarchy pStubHierarchy )
+		int IVsSolutionEvents.OnBeforeUnloadProject(IVsHierarchy pRealHierarchy, IVsHierarchy pStubHierarchy)
 		{
 			return VSConstants.S_OK;
 		}
 
-		int IVsSolutionEvents.OnQueryCloseProject( IVsHierarchy pHierarchy, int fRemoving, ref int pfCancel )
+		int IVsSolutionEvents.OnQueryCloseProject(IVsHierarchy pHierarchy, int fRemoving, ref int pfCancel)
 		{
 			return VSConstants.S_OK;
 		}
 
-		int IVsSolutionEvents.OnQueryCloseSolution( object pUnkReserved, ref int pfCancel )
+		int IVsSolutionEvents.OnQueryCloseSolution(object pUnkReserved, ref int pfCancel)
 		{
 			if (BatchBuilder.IsBusy)
 			{
@@ -660,7 +642,7 @@ namespace UnrealVS
 			return VSConstants.S_OK;
 		}
 
-		int IVsSolutionEvents.OnQueryUnloadProject( IVsHierarchy pRealHierarchy, ref int pfCancel )
+		int IVsSolutionEvents.OnQueryUnloadProject(IVsHierarchy pRealHierarchy, ref int pfCancel)
 		{
 			return VSConstants.S_OK;
 		}
@@ -670,24 +652,21 @@ namespace UnrealVS
 		/// IVsSelectionEvents implementation
 		///
 
-		int IVsSelectionEvents.OnCmdUIContextChanged( uint dwCmdUICookie, int fActive )
+		int IVsSelectionEvents.OnCmdUIContextChanged(uint dwCmdUICookie, int fActive)
 		{
-			if (OnUIContextChanged != null)
-			{
-				OnUIContextChanged(dwCmdUICookie, fActive != 0);
-			}
+			OnUIContextChanged?.Invoke(dwCmdUICookie, fActive != 0);
 
 			return VSConstants.S_OK;
 		}
 
-		int IVsSelectionEvents.OnElementValueChanged( uint elementid, object varValueOld, object varValueNew )
+		int IVsSelectionEvents.OnElementValueChanged(uint elementid, object varValueOld, object varValueNew)
 		{
 			// This function is called when selection changes in various Visual Studio tool windows
 			// and sub-systems.
 
 			// Handle startup project changes
-	        if( elementid == (uint)VSConstants.VSSELELEMID.SEID_StartupProject )
-            {
+			if (elementid == (uint)VSConstants.VSSELELEMID.SEID_StartupProject)
+			{
 				// If we are registered to a project hierarchy for events, unregister
 				var OldStartupProjectHierarchy = (IVsHierarchy)varValueOld;
 				if (OldStartupProjectHierarchy != null && ProjectHierarchyEventsHandle != 0)
@@ -700,10 +679,10 @@ namespace UnrealVS
 
 				// Incoming hierarchy object could be null (if no startup project is set yet, or during shutdown.)
 				var NewStartupProjectHierarchy = (IVsHierarchy)varValueNew;
-				if( NewStartupProjectHierarchy != null )
+				if (NewStartupProjectHierarchy != null)
 				{
 					// Get the actual Project object from the IVsHierarchy object that was supplied
-					NewStartupProject = Utils.HierarchyObjectToProject( NewStartupProjectHierarchy );
+					NewStartupProject = Utils.HierarchyObjectToProject(NewStartupProjectHierarchy);
 
 					if (NewStartupProject != null)
 					{
@@ -716,16 +695,16 @@ namespace UnrealVS
 				{
 					OnStartupProjectChanged(NewStartupProject);
 				}
-            }
+			}
 
 			return VSConstants.S_OK;
 		}
 
 		int IVsSelectionEvents.OnSelectionChanged(IVsHierarchy pHierOld, uint itemidOld, IVsMultiItemSelect pMISOld,
-		                                          ISelectionContainer pSCOld, IVsHierarchy pHierNew, uint itemidNew,
-		                                          IVsMultiItemSelect pMISNew, ISelectionContainer pSCNew)
+												  ISelectionContainer pSCOld, IVsHierarchy pHierNew, uint itemidNew,
+												  IVsMultiItemSelect pMISNew, ISelectionContainer pSCNew)
 		{
-    		return VSConstants.S_OK;
+			return VSConstants.S_OK;
 		}
 
 
@@ -738,10 +717,7 @@ namespace UnrealVS
 
 		Int32 IVsHierarchyEvents.OnPropertyChanged(UInt32 itemid, Int32 propid, UInt32 flags)
 		{
-			if (OnStartupProjectPropertyChanged != null)
-			{
-				OnStartupProjectPropertyChanged(itemid, propid, flags);
-			}
+			OnStartupProjectPropertyChanged?.Invoke(itemid, propid, flags);
 			return VSConstants.S_OK;
 		}
 
@@ -766,7 +742,7 @@ namespace UnrealVS
 		}
 
 
-#region IVsPersistSolutionProps
+		#region IVsPersistSolutionProps
 		int IVsPersistSolutionProps.SaveUserOptions(IVsSolutionPersistence pPersistence)
 		{
 			return VSConstants.S_OK;
@@ -807,8 +783,7 @@ namespace UnrealVS
 			if (!string.Equals(pszKey, GuidList.UnrealVSPackageString, StringComparison.InvariantCultureIgnoreCase))
 				return VSConstants.S_OK;
 
-			object availablePlatformsObject;
-			pPropBag.Read("AvailablePlatforms", out availablePlatformsObject, null, (uint)VarEnum.VT_BSTR, pPropBag);
+			pPropBag.Read("AvailablePlatforms", out object availablePlatformsObject, null, (uint)VarEnum.VT_BSTR, pPropBag);
 
 			string[] availablePlatforms = null;
 			if (availablePlatformsObject != null)
@@ -830,25 +805,19 @@ namespace UnrealVS
 			return VSConstants.S_OK;
 		}
 
-#endregion
+		#endregion
 
 		// IVsUpdateSolutionEvents Interface
 
 		public int UpdateSolution_Begin(ref int pfCancelUpdate)
 		{
-			if (OnBuildBegin != null)
-			{
-				OnBuildBegin(out pfCancelUpdate);
-			}
+			OnBuildBegin?.Invoke(out pfCancelUpdate);
 			return VSConstants.S_OK;
 		}
 
 		public int UpdateSolution_Done(int fSucceeded, int fModified, int fCancelCommand)
 		{
-			if (OnBuildDone != null)
-			{
-				OnBuildDone(fSucceeded != 0, fModified != 0, fCancelCommand != 0);
-			}
+			OnBuildDone?.Invoke(fSucceeded != 0, fModified != 0, fCancelCommand != 0);
 			return VSConstants.S_OK;
 		}
 
@@ -867,8 +836,7 @@ namespace UnrealVS
 			// This function is called after a Visual Studio project has its active config changed
 
 			// Check whether the project is the current startup project
-			IVsHierarchy StartupProjectHierarchy;
-			SolutionBuildManager.get_StartupProject(out StartupProjectHierarchy);
+			SolutionBuildManager.get_StartupProject(out IVsHierarchy StartupProjectHierarchy);
 			if (StartupProjectHierarchy != null && StartupProjectHierarchy == pIVsHierarchy)
 			{
 				// Get the actual Project object from the IVsHierarchy object that was supplied
@@ -889,15 +857,14 @@ namespace UnrealVS
 				return;
 			}
 
-			string SolutionDirectory, UserOptsFile;
-			SolutionManager.GetSolutionInfo(out SolutionDirectory, out _SolutionFilepath, out UserOptsFile);
+			SolutionManager.GetSolutionInfo(out string SolutionDirectory, out _SolutionFilepath, out string UserOptsFile);
 
 			// if ReadProps found a valud UE solution we do not need to check these legacy definitions of a solution
 			if (_IsUESolutionLoaded.HasValue)
 				return;
 
 			// Legacy paths for determing if a solution is a unreal solution by checking for tags & solution name.
-			var SolutionLines = new string[0];
+			string[] SolutionLines = Array.Empty<string>();
 			try
 			{
 				SolutionLines = File.ReadAllLines(_SolutionFilepath);
@@ -970,24 +937,22 @@ namespace UnrealVS
 		private static DTE2 GetDTE2ForCurrentInstance(DTE DTE)
 		{
 			// Find the ROT entry for visual studio running under current process.
-			IRunningObjectTable Rot;
-			NativeMethods.GetRunningObjectTable(0, out Rot);
-			IEnumMoniker EnumMoniker;
-			Rot.EnumRunning(out EnumMoniker);
-			EnumMoniker.Reset();
-			uint Fetched = 0;
-			IMoniker[] Moniker = new IMoniker[1];
-			while (EnumMoniker.Next(1, Moniker, out Fetched) == 0)
+			int HResult = NativeMethods.GetRunningObjectTable(0, out IRunningObjectTable Rot);
+			if (HResult == 0)
 			{
-				object ComObject;
-				Rot.GetObject(Moniker[0], out ComObject);
-				DTE2 CandidateDTE2 = ComObject as DTE2;
-
-				if (CandidateDTE2 != null)
+				Rot.EnumRunning(out IEnumMoniker EnumMoniker);
+				EnumMoniker.Reset();
+				IMoniker[] Moniker = new IMoniker[1];
+				while (EnumMoniker.Next(1, Moniker, out _) == 0)
 				{
-					if (CandidateDTE2.DTE == DTE)
+					Rot.GetObject(Moniker[0], out object ComObject);
+
+					if (ComObject is DTE2 CandidateDTE2)
 					{
-						return CandidateDTE2;
+						if (CandidateDTE2.DTE == DTE)
+						{
+							return CandidateDTE2;
+						}
 					}
 				}
 			}
diff --git a/Engine/Source/Programs/UnrealVS/UnrealVS_2019.csproj b/Engine/Source/Programs/UnrealVS/UnrealVS_2019.csproj
--- a/Engine/Source/Programs/UnrealVS/UnrealVS_2019.csproj
+++ b/Engine/Source/Programs/UnrealVS/UnrealVS_2019.csproj
@@ -44,31 +44,6 @@
     <ApplicationVersion>1.0.0.%2a</ApplicationVersion>
     <UseApplicationTrust>false</UseApplicationTrust>
     <BootstrapperEnabled>true</BootstrapperEnabled>
-    <NuGetPackageImportStamp>
-    </NuGetPackageImportStamp>
-  </PropertyGroup>
-  <PropertyGroup Condition=" '$(Configuration)|$(Platform)' == 'VS12d|AnyCPU' ">
-    <DebugSymbols>true</DebugSymbols>
-    <DebugType>full</DebugType>
-    <Optimize>false</Optimize>
-    <OutputPath>bin\VS12d\</OutputPath>
-    <DefineConstants>TRACE;DEBUG;VS12</DefineConstants>
-    <ErrorReport>prompt</ErrorReport>
-    <WarningLevel>4</WarningLevel>
-    <Prefer32Bit>false</Prefer32Bit>
-    <CopyVsixExtensionFiles>False</CopyVsixExtensionFiles>
-    <DeployExtension>True</DeployExtension>
-    <CreateVsixContainer>True</CreateVsixContainer>
-  </PropertyGroup>
-  <PropertyGroup Condition=" '$(Configuration)|$(Platform)' == 'VS12|AnyCPU' ">
-    <DebugType>pdbonly</DebugType>
-    <Optimize>true</Optimize>
-    <OutputPath>..\..\..\Extras\UnrealVS\VS2013\</OutputPath>
-    <DefineConstants>TRACE;CODE_ANALYSIS;VS12</DefineConstants>
-    <ErrorReport>prompt</ErrorReport>
-    <WarningLevel>4</WarningLevel>
-    <RunCodeAnalysis>true</RunCodeAnalysis>
-    <Prefer32Bit>false</Prefer32Bit>
   </PropertyGroup>
   <ItemGroup>
     <VSCTCompile Include="UnrealVS.vsct">
@@ -83,9 +58,6 @@
     <Reference Include="envdte80, Version=8.0.0.0, Culture=neutral, PublicKeyToken=b03f5f7f11d50a3a">
       <EmbedInteropTypes>True</EmbedInteropTypes>
     </Reference>
-    <Reference Include="Microsoft.Bcl.AsyncInterfaces, Version=1.0.0.0, Culture=neutral, PublicKeyToken=cc7b13ffcd2ddd51, processorArchitecture=MSIL">
-      <HintPath>packages\Microsoft.Bcl.AsyncInterfaces.1.1.0\lib\net461\Microsoft.Bcl.AsyncInterfaces.dll</HintPath>
-    </Reference>
     <Reference Include="Microsoft.CSharp" />
     <Reference Include="Microsoft.VisualStudio.OLE.Interop, Version=7.1.40304.0, Culture=neutral, PublicKeyToken=b03f5f7f11d50a3a" />
     <Reference Include="Microsoft.VisualStudio.Shell.15.0, Version=16.0.0.0, Culture=neutral, PublicKeyToken=b03f5f7f11d50a3a, processorArchitecture=MSIL" />
@@ -107,59 +79,20 @@
     <Reference Include="Microsoft.VisualStudio.Shell.Interop.8.0, Version=8.0.0.0, Culture=neutral, PublicKeyToken=b03f5f7f11d50a3a" />
     <Reference Include="Microsoft.VisualStudio.Shell.Interop.9.0, Version=9.0.0.0, Culture=neutral, PublicKeyToken=b03f5f7f11d50a3a" />
     <Reference Include="Microsoft.VisualStudio.TextManager.Interop, Version=7.1.40304.0, Culture=neutral, PublicKeyToken=b03f5f7f11d50a3a" />
-    <Reference Include="Microsoft.VisualStudio.Threading, Version=16.7.0.0, Culture=neutral, PublicKeyToken=b03f5f7f11d50a3a, processorArchitecture=MSIL">
-      <HintPath>packages\Microsoft.VisualStudio.Threading.16.7.56\lib\net472\Microsoft.VisualStudio.Threading.dll</HintPath>
-    </Reference>
-    <Reference Include="Microsoft.VisualStudio.Validation, Version=15.3.0.0, Culture=neutral, PublicKeyToken=b03f5f7f11d50a3a, processorArchitecture=MSIL">
-      <HintPath>packages\Microsoft.VisualStudio.Validation.15.5.31\lib\netstandard2.0\Microsoft.VisualStudio.Validation.dll</HintPath>
-    </Reference>
     <Reference Include="Microsoft.VisualStudio.VCProject, Version=16.0.0.0, Culture=neutral, PublicKeyToken=b03f5f7f11d50a3a">
       <EmbedInteropTypes>True</EmbedInteropTypes>
     </Reference>
     <Reference Include="Microsoft.VisualStudio.VCProjectEngine, Version=16.0.0.0, Culture=neutral, PublicKeyToken=b03f5f7f11d50a3a">
       <EmbedInteropTypes>True</EmbedInteropTypes>
     </Reference>
-    <Reference Include="Microsoft.Win32.Registry, Version=4.1.2.0, Culture=neutral, PublicKeyToken=b03f5f7f11d50a3a, processorArchitecture=MSIL">
-      <HintPath>packages\Microsoft.Win32.Registry.4.6.0\lib\net461\Microsoft.Win32.Registry.dll</HintPath>
-    </Reference>
     <Reference Include="PresentationCore" />
     <Reference Include="PresentationFramework" />
     <Reference Include="System" />
-    <Reference Include="System.Buffers, Version=4.0.3.0, Culture=neutral, PublicKeyToken=cc7b13ffcd2ddd51, processorArchitecture=MSIL">
-      <HintPath>packages\System.Buffers.4.5.1\lib\net461\System.Buffers.dll</HintPath>
-    </Reference>
     <Reference Include="System.Core" />
     <Reference Include="System.Data" />
     <Reference Include="System.Design" />
     <Reference Include="System.Drawing" />
-    <Reference Include="System.Memory, Version=4.0.1.1, Culture=neutral, PublicKeyToken=cc7b13ffcd2ddd51, processorArchitecture=MSIL">
-      <HintPath>packages\System.Memory.4.5.4\lib\net461\System.Memory.dll</HintPath>
-    </Reference>
     <Reference Include="System.Numerics" />
-    <Reference Include="System.Numerics.Vectors, Version=4.1.4.0, Culture=neutral, PublicKeyToken=b03f5f7f11d50a3a, processorArchitecture=MSIL">
-      <HintPath>packages\System.Numerics.Vectors.4.5.0\lib\net46\System.Numerics.Vectors.dll</HintPath>
-    </Reference>
-    <Reference Include="System.Runtime.CompilerServices.Unsafe, Version=4.0.6.0, Culture=neutral, PublicKeyToken=b03f5f7f11d50a3a, processorArchitecture=MSIL">
-      <HintPath>packages\System.Runtime.CompilerServices.Unsafe.4.7.1\lib\net461\System.Runtime.CompilerServices.Unsafe.dll</HintPath>
-    </Reference>
-    <Reference Include="System.Security.AccessControl, Version=4.1.2.0, Culture=neutral, PublicKeyToken=b03f5f7f11d50a3a, processorArchitecture=MSIL">
-      <HintPath>packages\System.Security.AccessControl.4.6.0\lib\net461\System.Security.AccessControl.dll</HintPath>
-    </Reference>
-    <Reference Include="System.Security.Principal.Windows, Version=4.1.2.0, Culture=neutral, PublicKeyToken=b03f5f7f11d50a3a, processorArchitecture=MSIL">
-      <HintPath>packages\System.Security.Principal.Windows.4.6.0\lib\net461\System.Security.Principal.Windows.dll</HintPath>
-    </Reference>
-    <Reference Include="System.Text.Encodings.Web, Version=4.0.5.0, Culture=neutral, PublicKeyToken=cc7b13ffcd2ddd51, processorArchitecture=MSIL">
-      <HintPath>packages\System.Text.Encodings.Web.4.7.1\lib\net461\System.Text.Encodings.Web.dll</HintPath>
-    </Reference>
-    <Reference Include="System.Text.Json, Version=4.0.1.2, Culture=neutral, PublicKeyToken=cc7b13ffcd2ddd51, processorArchitecture=MSIL">
-      <HintPath>packages\System.Text.Json.4.7.2\lib\net461\System.Text.Json.dll</HintPath>
-    </Reference>
-    <Reference Include="System.Threading.Tasks.Extensions, Version=4.2.0.1, Culture=neutral, PublicKeyToken=cc7b13ffcd2ddd51, processorArchitecture=MSIL">
-      <HintPath>packages\System.Threading.Tasks.Extensions.4.5.4\lib\net461\System.Threading.Tasks.Extensions.dll</HintPath>
-    </Reference>
-    <Reference Include="System.ValueTuple, Version=4.0.3.0, Culture=neutral, PublicKeyToken=cc7b13ffcd2ddd51, processorArchitecture=MSIL">
-      <HintPath>packages\System.ValueTuple.4.5.0\lib\net47\System.ValueTuple.dll</HintPath>
-    </Reference>
     <Reference Include="System.Windows.Forms" />
     <Reference Include="System.Xaml" />
     <Reference Include="System.Xml" />
@@ -230,7 +163,6 @@
   </ItemGroup>
   <ItemGroup>
     <None Include="app.config" />
-    <None Include="packages.config" />
   </ItemGroup>
   <ItemGroup>
     <None Include="Key.snk" />
@@ -291,12 +223,25 @@
     </Page>
   </ItemGroup>
   <ItemGroup>
-    <Analyzer Include="packages\Microsoft.VisualStudio.Threading.Analyzers.16.7.56\analyzers\cs\Microsoft.VisualStudio.Threading.Analyzers.CodeFixes.dll" />
-    <Analyzer Include="packages\Microsoft.VisualStudio.Threading.Analyzers.16.7.56\analyzers\cs\Microsoft.VisualStudio.Threading.Analyzers.CSharp.dll" />
-    <Analyzer Include="packages\Microsoft.VisualStudio.Threading.Analyzers.16.7.56\analyzers\cs\Microsoft.VisualStudio.Threading.Analyzers.dll" />
-    <Analyzer Include="packages\Microsoft.VisualStudio.Threading.Analyzers.16.7.56\analyzers\vb\Microsoft.VisualStudio.Threading.Analyzers.CodeFixes.dll" />
-    <Analyzer Include="packages\Microsoft.VisualStudio.Threading.Analyzers.16.7.56\analyzers\vb\Microsoft.VisualStudio.Threading.Analyzers.dll" />
-    <Analyzer Include="packages\Microsoft.VisualStudio.Threading.Analyzers.16.7.56\analyzers\vb\Microsoft.VisualStudio.Threading.Analyzers.VisualBasic.dll" />
+    <PackageReference Include="Microsoft.CodeAnalysis.CSharp">
+      <Version>3.9.0</Version>
+    </PackageReference>
+    <PackageReference Include="Microsoft.CodeAnalysis.NetAnalyzers">
+      <Version>5.0.3</Version>
+      <IncludeAssets>runtime; build; native; contentfiles; analyzers; buildtransitive</IncludeAssets>
+      <PrivateAssets>all</PrivateAssets>
+    </PackageReference>
+    <PackageReference Include="Microsoft.VisualStudio.Threading">
+      <Version>16.9.60</Version>
+    </PackageReference>
+    <PackageReference Include="Microsoft.VisualStudio.Threading.Analyzers">
+      <Version>16.9.60</Version>
+      <IncludeAssets>runtime; build; native; contentfiles; analyzers; buildtransitive</IncludeAssets>
+      <PrivateAssets>all</PrivateAssets>
+    </PackageReference>
+    <PackageReference Include="System.Text.Json">
+      <Version>5.0.2</Version>
+    </PackageReference>
   </ItemGroup>
   <PropertyGroup>
     <UseCodebase>true</UseCodebase>
@@ -305,81 +250,19 @@
     <VisualStudioVersion Condition="'$(VisualStudioVersion)' == ''">11.0</VisualStudioVersion>
     <VSToolsPath Condition="'$(VSToolsPath)' == ''">$(MSBuildExtensionsPath32)\Microsoft\VisualStudio\v$(VisualStudioVersion)</VSToolsPath>
   </PropertyGroup>
-  <PropertyGroup Condition="'$(Configuration)|$(Platform)' == 'VS11|AnyCPU'">
-    <DebugType>pdbonly</DebugType>
-    <Optimize>true</Optimize>
-    <OutputPath>..\..\..\Extras\UnrealVS\VS2012\</OutputPath>
-    <DefineConstants>TRACE;CODE_ANALYSIS;VS11</DefineConstants>
-    <ErrorReport>prompt</ErrorReport>
-    <WarningLevel>4</WarningLevel>
-    <RunCodeAnalysis>true</RunCodeAnalysis>
-    <Prefer32Bit>false</Prefer32Bit>
-  </PropertyGroup>
-  <PropertyGroup Condition=" '$(Configuration)|$(Platform)' == 'VS11d|AnyCPU' ">
-    <DebugSymbols>true</DebugSymbols>
-    <DebugType>full</DebugType>
-    <Optimize>false</Optimize>
-    <OutputPath>bin\VS11d\</OutputPath>
-    <DefineConstants>TRACE;DEBUG;VS11</DefineConstants>
-    <ErrorReport>prompt</ErrorReport>
-    <WarningLevel>4</WarningLevel>
-    <Prefer32Bit>false</Prefer32Bit>
-    <CopyVsixExtensionFiles>False</CopyVsixExtensionFiles>
-    <DeployExtension>True</DeployExtension>
-    <CreateVsixContainer>True</CreateVsixContainer>
-  </PropertyGroup>
-  <PropertyGroup Condition=" '$(Configuration)|$(Platform)' == 'VS14d|AnyCPU' ">
-    <DebugSymbols>true</DebugSymbols>
-    <OutputPath>bin\VS14d\</OutputPath>
-    <DefineConstants>TRACE;DEBUG;VS14</DefineConstants>
-    <DebugType>full</DebugType>
-    <PlatformTarget>AnyCPU</PlatformTarget>
-    <ErrorReport>prompt</ErrorReport>
-    <CodeAnalysisRuleSet>MinimumRecommendedRules.ruleset</CodeAnalysisRuleSet>
-  </PropertyGroup>
-  <PropertyGroup Condition=" '$(Configuration)|$(Platform)' == 'VS14|AnyCPU' ">
-    <OutputPath>..\..\..\Extras\UnrealVS\VS2015\</OutputPath>
-    <DefineConstants>TRACE;CODE_ANALYSIS;VS14</DefineConstants>
-    <Optimize>true</Optimize>
-    <DebugType>pdbonly</DebugType>
-    <PlatformTarget>AnyCPU</PlatformTarget>
-    <RunCodeAnalysis>true</RunCodeAnalysis>
-    <ErrorReport>prompt</ErrorReport>
-    <CodeAnalysisRuleSet>MinimumRecommendedRules.ruleset</CodeAnalysisRuleSet>
-  </PropertyGroup>
-  <PropertyGroup Condition="'$(Configuration)|$(Platform)' == 'VS15|AnyCPU'">
-    <OutputPath>..\..\..\Extras\UnrealVS\VS2017\</OutputPath>
-    <DefineConstants>CODE_ANALYSIS;TRACE;CODE_ANALYSIS;VS14</DefineConstants>
-    <Optimize>true</Optimize>
-    <DebugType>pdbonly</DebugType>
-    <PlatformTarget>AnyCPU</PlatformTarget>
-    <RunCodeAnalysis>true</RunCodeAnalysis>
-    <ErrorReport>prompt</ErrorReport>
-    <CodeAnalysisRuleSet>MinimumRecommendedRules.ruleset</CodeAnalysisRuleSet>
-  </PropertyGroup>
-  <PropertyGroup Condition="'$(Configuration)|$(Platform)' == 'VS15d|AnyCPU'">
-    <DebugSymbols>true</DebugSymbols>
-    <OutputPath>bin\VS15d\</OutputPath>
-    <DefineConstants>TRACE;DEBUG;VS14</DefineConstants>
-    <DebugType>full</DebugType>
-    <PlatformTarget>AnyCPU</PlatformTarget>
-    <ErrorReport>prompt</ErrorReport>
-    <CodeAnalysisRuleSet>MinimumRecommendedRules.ruleset</CodeAnalysisRuleSet>
-    <DeployExtension>True</DeployExtension>
-    <CopyVsixExtensionFiles>False</CopyVsixExtensionFiles>
-  </PropertyGroup>
   <PropertyGroup Condition=" '$(Configuration)|$(Platform)' == 'VS16d|AnyCPU' ">
     <DebugSymbols>true</DebugSymbols>
     <OutputPath>..\..\..\Extras\UnrealVS\VS2019\</OutputPath>
-    <DefineConstants>TRACE;DEBUG;VS14</DefineConstants>
+    <DefineConstants>TRACE;DEBUG;VS16</DefineConstants>
     <DebugType>full</DebugType>
     <PlatformTarget>AnyCPU</PlatformTarget>
     <ErrorReport>prompt</ErrorReport>
     <CodeAnalysisRuleSet>MinimumRecommendedRules.ruleset</CodeAnalysisRuleSet>
+    <TreatWarningsAsErrors>true</TreatWarningsAsErrors>
   </PropertyGroup>
   <PropertyGroup Condition=" '$(Configuration)|$(Platform)' == 'VS16|AnyCPU' ">
     <OutputPath>..\..\..\Extras\UnrealVS\VS2019\</OutputPath>
-    <DefineConstants>CODE_ANALYSIS;CODE_ANALYSIS;TRACE;CODE_ANALYSIS;VS14</DefineConstants>
+    <DefineConstants>TRACE;VS16</DefineConstants>
     <Optimize>true</Optimize>
     <DebugType>pdbonly</DebugType>
     <PlatformTarget>AnyCPU</PlatformTarget>
@@ -387,6 +270,7 @@
     <ErrorReport>prompt</ErrorReport>
     <CodeAnalysisRuleSet>MinimumRecommendedRules.ruleset</CodeAnalysisRuleSet>
     <CopyVsixExtensionFiles>False</CopyVsixExtensionFiles>
+    <TreatWarningsAsErrors>true</TreatWarningsAsErrors>
   </PropertyGroup>
   <Import Project="$(MSBuildBinPath)\Microsoft.CSharp.targets" />
   <Import Project="$(VSToolsPath)\VSSDK\Microsoft.VsSDK.targets" Condition="'$(VSToolsPath)' != ''" />
@@ -398,13 +282,6 @@
   <PropertyGroup>
     <PreBuildEvent>copy "$(ProjectDir)TargetManifests\$(ConfigurationName).source.extension.vsixmanifest" "$(ProjectDir)Generated\source.extension.vsixmanifest" /Y</PreBuildEvent>
   </PropertyGroup>
-  <Import Project="packages\Microsoft.VisualStudio.Threading.Analyzers.16.7.56\build\Microsoft.VisualStudio.Threading.Analyzers.targets" Condition="Exists('packages\Microsoft.VisualStudio.Threading.Analyzers.16.7.56\build\Microsoft.VisualStudio.Threading.Analyzers.targets')" />
-  <Target Name="EnsureNuGetPackageBuildImports" BeforeTargets="PrepareForBuild">
-    <PropertyGroup>
-      <ErrorText>This project references NuGet package(s) that are missing on this computer. Use NuGet Package Restore to download them.  For more information, see http://go.microsoft.com/fwlink/?LinkID=322105. The missing file is {0}.</ErrorText>
-    </PropertyGroup>
-    <Error Condition="!Exists('packages\Microsoft.VisualStudio.Threading.Analyzers.16.7.56\build\Microsoft.VisualStudio.Threading.Analyzers.targets')" Text="$([System.String]::Format('$(ErrorText)', 'packages\Microsoft.VisualStudio.Threading.Analyzers.16.7.56\build\Microsoft.VisualStudio.Threading.Analyzers.targets'))" />
-  </Target>
   <!-- To modify your build process, add your task inside one of the targets below and uncomment it. 
        Other similar extension points exist, see Microsoft.Common.targets.
   <Target Name="BeforeBuild">
diff --git a/Engine/Source/Programs/UnrealVS/Utils.cs b/Engine/Source/Programs/UnrealVS/Utils.cs
--- a/Engine/Source/Programs/UnrealVS/Utils.cs
+++ b/Engine/Source/Programs/UnrealVS/Utils.cs
@@ -1,16 +1,13 @@
-﻿// Copyright Epic Games, Inc. All Rights Reserved.
+// Copyright Epic Games, Inc. All Rights Reserved.
 
-using System;
+using EnvDTE;
+using EnvDTE80;
 using Microsoft.VisualStudio;
 using Microsoft.VisualStudio.Shell.Interop;
-using EnvDTE;
+using System;
 using System.Collections.Generic;
-using System.Linq;
-using System.Linq.Expressions;
-using VSLangProj;
 using System.IO;
-using EnvDTE80;
-using Microsoft.VisualStudio.Shell;
+using System.Linq;
 using System.Xml;
 
 namespace UnrealVS
@@ -36,10 +33,9 @@ namespace UnrealVS
 		/// </summary>
 		/// <param name="Project">Project object</param>
 		/// <returns>IVsHierarchy for the specified project</returns>
-		public static IVsHierarchy ProjectToHierarchyObject( Project Project )
+		public static IVsHierarchy ProjectToHierarchyObject(Project Project)
 		{
-			IVsHierarchy HierarchyObject;
-			UnrealVSPackage.Instance.SolutionManager.GetProjectOfUniqueName(Project.FullName, out HierarchyObject);
+			UnrealVSPackage.Instance.SolutionManager.GetProjectOfUniqueName(Project.FullName, out IVsHierarchy HierarchyObject);
 			return HierarchyObject;
 		}
 
@@ -49,11 +45,10 @@ namespace UnrealVS
 		/// </summary>
 		/// <param name="HierarchyObject">IVsHierarchy object</param>
 		/// <returns>Visual Studio project object</returns>
-		public static Project HierarchyObjectToProject( IVsHierarchy HierarchyObject )
+		public static Project HierarchyObjectToProject(IVsHierarchy HierarchyObject)
 		{
 			// Get the actual Project object from the IVsHierarchy object that was supplied
-			object ProjectObject;
-			HierarchyObject.GetProperty(VSConstants.VSITEMID_ROOT, (int)__VSHPROPID.VSHPROPID_ExtObject, out ProjectObject);
+			HierarchyObject.GetProperty(VSConstants.VSITEMID_ROOT, (int)__VSHPROPID.VSHPROPID_ExtObject, out object ProjectObject);
 			return (Project)ProjectObject;
 		}
 
@@ -62,11 +57,10 @@ namespace UnrealVS
 		/// </summary>
 		/// <param name="HierarchyObject">IVsHierarchy object</param>
 		/// <returns>Visual Studio project object</returns>
-		public static IVsCfgProvider2 HierarchyObjectToCfgProvider( IVsHierarchy HierarchyObject )
+		public static IVsCfgProvider2 HierarchyObjectToCfgProvider(IVsHierarchy HierarchyObject)
 		{
 			// Get the actual Project object from the IVsHierarchy object that was supplied
-			object BrowseObject;
-			HierarchyObject.GetProperty(VSConstants.VSITEMID_ROOT, (int)__VSHPROPID.VSHPROPID_BrowseObject, out BrowseObject);
+			HierarchyObject.GetProperty(VSConstants.VSITEMID_ROOT, (int)__VSHPROPID.VSHPROPID_BrowseObject, out object BrowseObject);
 
 			IVsCfgProvider2 CfgProvider = null;
 			if (BrowseObject != null)
@@ -82,15 +76,13 @@ namespace UnrealVS
 			return CfgProvider;
 		}
 
-		private static IVsCfgProvider2 GetCfgProviderFromObject( object SomeObject )
+		private static IVsCfgProvider2 GetCfgProviderFromObject(object SomeObject)
 		{
 			IVsCfgProvider2 CfgProvider2 = null;
 
-			var GetCfgProvider = SomeObject as IVsGetCfgProvider;
-			if (GetCfgProvider != null)
+			if (SomeObject is IVsGetCfgProvider GetCfgProvider)
 			{
-				IVsCfgProvider CfgProvider;
-				GetCfgProvider.GetCfgProvider(out CfgProvider);
+				GetCfgProvider.GetCfgProvider(out IVsCfgProvider CfgProvider);
 				if (CfgProvider != null)
 				{
 					CfgProvider2 = CfgProvider as IVsCfgProvider2;
@@ -111,7 +103,7 @@ namespace UnrealVS
 		/// <param name="Project">Project to search for the property</param>
 		/// <param name="PropertyName">Name of the property</param>
 		/// <returns>Property object or null if not found</returns>
-		public static Property GetProjectProperty( Project Project, string PropertyName )
+		public static Property GetProjectProperty(Project Project, string PropertyName)
 		{
 			var Properties = Project.Properties;
 			if (Properties != null)
@@ -135,7 +127,7 @@ namespace UnrealVS
 		/// </summary>
 		/// <param name="Property">The property object to set</param>
 		/// <param name="PropertyValue">Value to set for this property</param>
-		public static void SetPropertyValue( Property Property, object PropertyValue )
+		public static void SetPropertyValue(Property Property, object PropertyValue)
 		{
 			Property.Value = PropertyValue;
 
@@ -155,13 +147,13 @@ namespace UnrealVS
 			public UIHierarchyItem Item { get; set; }
 			public UITreeItem[] Children { get; set; }
 			public string Name { get { return Item != null ? Item.Name : "None"; } }
-			public object Object { get { return Item != null ? Item.Object : null; } }
+			public object Object { get { return Item?.Object; } }
 		}
 
 		/// <summary>
 		/// Converts a UIHierarchy into an easy to use tree of helper class UITreeItem.
 		/// </summary>
-		public static UITreeItem GetUIHierarchyTree( UIHierarchy Hierarchy )
+		public static UITreeItem GetUIHierarchyTree(UIHierarchy Hierarchy)
 		{
 			return new UITreeItem
 			{
@@ -173,7 +165,7 @@ namespace UnrealVS
 		/// <summary>
 		/// Called by the public GetUIHierarchyTree() function above.
 		/// </summary>
-		private static UITreeItem GetUIHierarchyTree( UIHierarchyItem HierarchyItem )
+		private static UITreeItem GetUIHierarchyTree(UIHierarchyItem HierarchyItem)
 		{
 			return new UITreeItem
 			{
@@ -188,13 +180,13 @@ namespace UnrealVS
 		/// <typeparam name="T">The type of object to find in the tree. Extracts everything that "Is a" T.</typeparam>
 		/// <param name="RootItem">The root of the UIHierarchy to search (converted to UITreeItem via GetUIHierarchyTree())</param>
 		/// <returns>An enumerable of objects of type T found beneath the root item.</returns>
-		public static IEnumerable<T> GetUITreeItemObjectsByType<T>( UITreeItem RootItem ) where T : class
+		public static IEnumerable<T> GetUITreeItemObjectsByType<T>(UITreeItem RootItem) where T : class
 		{
 			List<T> Results = new List<T>();
 
-			if (RootItem.Object is T)
+			if (RootItem.Object is T Obj)
 			{
-				Results.Add((T)RootItem.Object);
+				Results.Add(Obj);
 			}
 			foreach (var Child in RootItem.Children)
 			{
@@ -204,7 +196,7 @@ namespace UnrealVS
 			return Results;
 		}
 
-		public static IEnumerable<UIHierarchyItem> GetUITreeItemsByObjectType<T>( UITreeItem RootItem ) where T : class
+		public static IEnumerable<UIHierarchyItem> GetUITreeItemsByObjectType<T>(UITreeItem RootItem) where T : class
 		{
 			List<UIHierarchyItem> Results = new List<UIHierarchyItem>();
 
@@ -224,7 +216,7 @@ namespace UnrealVS
 		/// Helper to check the file ext of a binary against known library file exts.
 		/// FileExt should include the dot e.g. ".dll"
 		/// </summary>
-		public static bool IsLibraryFileExtension( string FileExt )
+		public static bool IsLibraryFileExtension(string FileExt)
 		{
 			if (FileExt.Equals(".dll", StringComparison.InvariantCultureIgnoreCase)) return true;
 			if (FileExt.Equals(".lib", StringComparison.InvariantCultureIgnoreCase)) return true;
@@ -239,7 +231,7 @@ namespace UnrealVS
 		/// <summary>
 		/// Helper to check the properties of a project and determine whether it can be built in VS.
 		/// </summary>
-		public static bool IsProjectBuildable( Project Project )
+		public static bool IsProjectBuildable(Project Project)
 		{
 			return Project.Kind == GuidList.VCSharpProjectKindGuidString || Project.Kind == GuidList.VCProjectKindGuidString;
 		}
@@ -275,12 +267,12 @@ namespace UnrealVS
 			}
 		}
 
-		public static void ExecuteProjectBuild( Project Project,
+		public static void ExecuteProjectBuild(Project Project,
 												string SolutionConfig,
 												string SolutionPlatform,
 												BatchBuilderToolControl.BuildJob.BuildJobType BuildType,
 												Action ExecutingDelegate,
-												Action FailedToStartDelegate )
+												Action FailedToStartDelegate)
 		{
 			IVsHierarchy ProjHierarchy = Utils.ProjectToHierarchyObject(Project);
 
@@ -303,12 +295,11 @@ namespace UnrealVS
 						IVsCfgProvider2 CfgProvider2 = Utils.HierarchyObjectToCfgProvider(ProjHierarchy);
 						if (CfgProvider2 != null)
 						{
-							IVsCfg Cfg;
-							CfgProvider2.GetCfgOfName(ProjectSolutionCtxt.ConfigurationName, ProjectSolutionCtxt.PlatformName, out Cfg);
+							CfgProvider2.GetCfgOfName(ProjectSolutionCtxt.ConfigurationName, ProjectSolutionCtxt.PlatformName, out IVsCfg Cfg);
 
 							if (Cfg != null)
 							{
-								if (ExecutingDelegate != null) ExecutingDelegate();
+								ExecutingDelegate?.Invoke();
 
 								int JobResult = VSConstants.E_FAIL;
 
@@ -359,7 +350,7 @@ namespace UnrealVS
 								}
 								else
 								{
-									if (FailedToStartDelegate != null) FailedToStartDelegate();
+									FailedToStartDelegate?.Invoke();
 								}
 							}
 						}
@@ -369,7 +360,7 @@ namespace UnrealVS
 		}
 
 
-		private static bool LoadConfigFromUBT( Project SelectedProject )
+		private static bool LoadConfigFromUBT(Project SelectedProject)
 		{
 			string ProjectPath = Path.GetDirectoryName(SelectedProject.FullName);
 			string ConfigFileName = Path.Combine(ProjectPath, "UnrealVS.xml");
@@ -397,7 +388,7 @@ namespace UnrealVS
 		}
 
 
-		public static List<string> GetExtraDebuggerCommandArguments( string PlatformName, Project SelectedProject )
+		public static List<string> GetExtraDebuggerCommandArguments(string PlatformName, Project SelectedProject)
 		{
 			List<string> Result = new List<string>();
 
@@ -442,8 +433,7 @@ namespace UnrealVS
 			var UProjectFileName = GetUProjectFileName(Project);
 			var AllUProjects = GetUProjects();
 
-			string UProjectPath = string.Empty;
-			if (!AllUProjects.TryGetValue(Project.Name, out UProjectPath))
+			if (!AllUProjects.TryGetValue(Project.Name, out string UProjectPath))
 			{
 				// Search the project folder
 				var ProjectFolder = Path.GetDirectoryName(Project.FullName);
@@ -451,7 +441,7 @@ namespace UnrealVS
 				if (UProjUnderProject.Length == 1)
 				{
 					UProjectPath = UProjUnderProject[0];
-				}				
+				}
 			}
 
 			return '\"' + UProjectPath + '\"';
@@ -474,8 +464,10 @@ namespace UnrealVS
 			if (ProjectFiles.Count == 0)
 			{
 				// Build a list of all the parent directories for projects. This includes the UE4 root, plus any directories referenced via .uprojectdirs files.
-				List<DirectoryInfo> ParentProjectDirs = new List<DirectoryInfo>();
-				ParentProjectDirs.Add(SolutionDir);
+				List<DirectoryInfo> ParentProjectDirs = new List<DirectoryInfo>
+				{
+					SolutionDir
+				};
 
 				// Read all the .uprojectdirs files
 				foreach (FileInfo ProjectDirsFile in SolutionDir.EnumerateFiles("*.uprojectdirs"))
@@ -501,7 +493,7 @@ namespace UnrealVS
 				HashSet<string> CheckedParentDirs = new HashSet<string>(StringComparer.InvariantCultureIgnoreCase);
 				foreach (DirectoryInfo ParentProjectDir in ParentProjectDirs)
 				{
-					if(CheckedParentDirs.Add(ParentProjectDir.FullName) && ParentProjectDir.Exists)
+					if (CheckedParentDirs.Add(ParentProjectDir.FullName) && ParentProjectDir.Exists)
 					{
 						foreach (DirectoryInfo ProjectDir in ParentProjectDir.EnumerateDirectories())
 						{
@@ -527,21 +519,21 @@ namespace UnrealVS
 			if (Folder != CachedUProjectRootFolder)
 			{
 				Logging.WriteLine("GetUProjects: recaching uproject paths...");
-                DateTime Start = DateTime.Now;
+				DateTime Start = DateTime.Now;
 
 				CachedUProjectRootFolder = Folder;
 				CachedUProjectPaths = EnumerateProjects(new DirectoryInfo(Folder)).Select(x => x.FullName);
 				CachedUProjects = null;
 
-                TimeSpan TimeTaken = DateTime.Now - Start;
-                Logging.WriteLine(string.Format("GetUProjects: EnumerateProjects took {0} sec", TimeTaken.TotalSeconds));
+				TimeSpan TimeTaken = DateTime.Now - Start;
+				Logging.WriteLine(string.Format("GetUProjects: EnumerateProjects took {0} sec", TimeTaken.TotalSeconds));
 
 				foreach (string CachedUProjectPath in CachedUProjectPaths)
 				{
 					Logging.WriteLine(String.Format("GetUProjects: found {0}", CachedUProjectPath));
 				}
 
-                Logging.WriteLine("    DONE");
+				Logging.WriteLine("    DONE");
 			}
 
 			if (CachedUProjects == null)
@@ -552,9 +544,9 @@ namespace UnrealVS
 				var ProjectNames = (from path in ProjectPaths select Path.GetFileNameWithoutExtension(path)).ToArray();
 
 				var CodeUProjects = from UProjectPath in CachedUProjectPaths
-					let ProjectName = Path.GetFileNameWithoutExtension(UProjectPath)
-					where ProjectNames.Any(name => string.Compare(name, ProjectName, StringComparison.OrdinalIgnoreCase) == 0)
-					select new {Name = ProjectName, FilePath = UProjectPath};
+									let ProjectName = Path.GetFileNameWithoutExtension(UProjectPath)
+									where ProjectNames.Any(name => string.Compare(name, ProjectName, StringComparison.OrdinalIgnoreCase) == 0)
+									select new { Name = ProjectName, FilePath = UProjectPath };
 
 				CachedUProjects = new Dictionary<string, string>();
 
diff --git a/Engine/Source/Programs/UnrealVS/packages.config b/Engine/Source/Programs/UnrealVS/packages.config
deleted file mode 100644
