UnrealVS: Remove obsolete solutions and installers
Update references of UE4 -> UE5
Update more info url to point at documentation https://docs.unrealengine.com/en-US/ProductionPipelines/DevelopmentSetup/VisualStudioSetup/UnrealVS/index.html

diff --git a/Engine/Source/Programs/UnrealBuildTool/UnrealBuildTool.csproj b/Engine/Source/Programs/UnrealBuildTool/UnrealBuildTool.csproj
--- a/Engine/Source/Programs/UnrealBuildTool/UnrealBuildTool.csproj
+++ b/Engine/Source/Programs/UnrealBuildTool/UnrealBuildTool.csproj
@@ -32,12 +32,12 @@
     </When>
   </Choose>
   <PropertyGroup Condition=" '$(Configuration)|$(Platform)' == 'Development|AnyCPU' ">
-    <Optimize>true</Optimize>
+    <Optimize>false</Optimize>
     <OutputPath>..\..\..\Binaries\DotNET\UnrealBuildTool</OutputPath>
     <DocumentationFile>..\..\..\Binaries\DotNET\UnrealBuildTool\UnrealBuildTool.xml</DocumentationFile>
   </PropertyGroup>
   <PropertyGroup Condition=" '$(Configuration)|$(Platform)' == 'Release|AnyCPU' ">
-    <Optimize>true</Optimize>
+    <Optimize>false</Optimize>
     <DebugType>pdbonly</DebugType>
     <PlatformTarget>AnyCPU</PlatformTarget>
     <ErrorReport>prompt</ErrorReport>
diff --git a/Engine/Source/Programs/UnrealVS/CommandLineEditor.cs b/Engine/Source/Programs/UnrealVS/CommandLineEditor.cs
--- a/Engine/Source/Programs/UnrealVS/CommandLineEditor.cs
+++ b/Engine/Source/Programs/UnrealVS/CommandLineEditor.cs
@@ -75,7 +75,7 @@ namespace UnrealVS
 		/// </summary>
 		public void OnStartupProjectPropertyChanged(UInt32 itemid, Int32 propid, UInt32 flags)
 		{
-			// Filter out Helix VS plugin sending thousands of VSHPROPID_StateIconIndex in UE4.sln
+			// Filter out Helix VS plugin sending thousands of VSHPROPID_StateIconIndex in the solution
 			// This event type is not usefull at all for the commandline editor
 			if (propid == (Int32)__VSHPROPID.VSHPROPID_StateIconIndex)
 			{
diff --git a/Engine/Source/Programs/UnrealVS/ExecuteCIS.cs b/Engine/Source/Programs/UnrealVS/ExecuteCIS.cs
deleted file mode 100644
diff --git a/Engine/Source/Programs/UnrealVS/GenerateProjectFiles.cs b/Engine/Source/Programs/UnrealVS/GenerateProjectFiles.cs
--- a/Engine/Source/Programs/UnrealVS/GenerateProjectFiles.cs
+++ b/Engine/Source/Programs/UnrealVS/GenerateProjectFiles.cs
@@ -96,7 +96,7 @@ namespace UnrealVS
 					}
 					else
 					{
-						// No solution loaded, or solution doesn't look like UE4
+						// No solution loaded, or solution doesn't look like Unreal
 					}
 				}
 			}
@@ -137,7 +137,7 @@ namespace UnrealVS
 
 		private string GetBatchFileName()
 		{
-			// Check to see if we have UE4.sln loaded
+			// Check to see if we have the Unreal loaded
 			if (UnrealVSPackage.Instance.IsUESolutionLoaded)
 			{
 				// We expect "GenerateProjectFiles.bat" to live in the same directory as the solution
diff --git a/Engine/Source/Programs/UnrealVS/StartupProjectSelector.cs b/Engine/Source/Programs/UnrealVS/StartupProjectSelector.cs
--- a/Engine/Source/Programs/UnrealVS/StartupProjectSelector.cs
+++ b/Engine/Source/Programs/UnrealVS/StartupProjectSelector.cs
@@ -150,7 +150,7 @@ namespace UnrealVS
 		private void SortCachedStartupProjects()
 		{
 			// Sort projects by game y/n then alphabetically
-			_CachedStartupProjects = (_CachedStartupProjects.OrderBy(ProjectRef => ProjectRef.Name == "UE4" ? 0 : 1)
+			_CachedStartupProjects = (_CachedStartupProjects.OrderBy(ProjectRef => ProjectRef.Name == "UE5" ? 0 : 1)
 				.ThenBy(ProjectRef => Utils.IsGameProject(ProjectRef.Project) ? 0 : 1)
 				.ThenBy(ProjectRef => ProjectRef.Name)).ToList();
 		}
diff --git a/Engine/Source/Programs/UnrealVS/UnrealVS.csproj b/Engine/Source/Programs/UnrealVS/UnrealVS.csproj
deleted file mode 100644
diff --git a/Engine/Source/Programs/UnrealVS/UnrealVSPackage.cs b/Engine/Source/Programs/UnrealVS/UnrealVSPackage.cs
--- a/Engine/Source/Programs/UnrealVS/UnrealVSPackage.cs
+++ b/Engine/Source/Programs/UnrealVS/UnrealVSPackage.cs
@@ -69,8 +69,8 @@ namespace UnrealVS
 	{
 		/** Constants */
 
-		private const string VersionString = "v1.57";
-		private const string UnrealSolutionFileNamePrefix = "UE4";
+		private const string VersionString = "v1.60";
+		private const string UnrealSolutionFileNamePrefix = "UE";
 		private const string ExtensionName = "UnrealVS";
 		private const string CommandLineOptionKey = ExtensionName + "CommandLineMRU";
 		private const string BatchBuildSetsOptionKey = ExtensionName + "BatchBuildSetsV002";
diff --git a/Engine/Source/Programs/UnrealVS/UnrealVS_2017.csproj b/Engine/Source/Programs/UnrealVS/UnrealVS_2017.csproj
deleted file mode 100644
diff --git a/Engine/Source/Programs/UnrealVS/UnrealVS_2019.csproj b/Engine/Source/Programs/UnrealVS/UnrealVS_2019.csproj
--- a/Engine/Source/Programs/UnrealVS/UnrealVS_2019.csproj
+++ b/Engine/Source/Programs/UnrealVS/UnrealVS_2019.csproj
@@ -168,15 +168,12 @@
     <None Include="Key.snk" />
   </ItemGroup>
   <ItemGroup>
-    <Content Include="TargetManifests\VS14.source.extension.vsixmanifest" />
-    <Content Include="TargetManifests\VS14d.source.extension.vsixmanifest" />
-    <Content Include="TargetManifests\VS15.source.extension.vsixmanifest" />
-    <Content Include="TargetManifests\VS15d.source.extension.vsixmanifest" />
     <Content Include="TargetManifests\VS16.source.extension.vsixmanifest" />
     <Content Include="TargetManifests\VS16d.source.extension.vsixmanifest" />
     <None Include="Generated\source.extension.vsixmanifest" />
     <None Include="Resources\BatchBuilder.png" />
     <Resource Include="Resources\BatchBuilderBackground.png" />
+    <Resource Include="Resources\Icon.ico" />
     <Content Include="Resources\Icon.png">
       <CopyToOutputDirectory>Always</CopyToOutputDirectory>
       <IncludeInVSIX>true</IncludeInVSIX>
diff --git a/Engine/Source/Programs/UnrealVS/Utils.cs b/Engine/Source/Programs/UnrealVS/Utils.cs
--- a/Engine/Source/Programs/UnrealVS/Utils.cs
+++ b/Engine/Source/Programs/UnrealVS/Utils.cs
@@ -463,7 +463,7 @@ namespace UnrealVS
 			List<FileInfo> ProjectFiles = new List<FileInfo>(SolutionDir.EnumerateFiles("*.uproject"));
 			if (ProjectFiles.Count == 0)
 			{
-				// Build a list of all the parent directories for projects. This includes the UE4 root, plus any directories referenced via .uprojectdirs files.
+				// Build a list of all the parent directories for projects. This includes the UE root, plus any directories referenced via .uprojectdirs files.
 				List<DirectoryInfo> ParentProjectDirs = new List<DirectoryInfo>
 				{
 					SolutionDir
