UnrealVS: Move project to UnrealVS.2019\ to prepare for making the 2022 extension
Update all Microsoft.VisualStudio.SDK dependencies to use nuget packages
Reset output directory and copy the .vsix to the target location in the post build step
Remove need for multiple source.extension.vsixmanifest files

diff --git a/Engine/Source/Programs/UnrealVS/.p4ignore.txt b/Engine/Source/Programs/UnrealVS/.p4ignore.txt
new file mode 100644
--- /dev/null
+++ b/Engine/Source/Programs/UnrealVS/.p4ignore.txt
@@ -0,0 +1,3 @@
+**/bin/Release
+**/bin/Debug
+*.csproj.user
\ No newline at end of file
diff --git a/Engine/Source/Programs/UnrealVS/BatchBuilder.cs b/Engine/Source/Programs/UnrealVS/BatchBuilder.cs
deleted file mode 100644
diff --git a/Engine/Source/Programs/UnrealVS/BatchBuilderToolControl.xaml b/Engine/Source/Programs/UnrealVS/BatchBuilderToolControl.xaml
deleted file mode 100644
diff --git a/Engine/Source/Programs/UnrealVS/BatchBuilderToolControl.xaml.cs b/Engine/Source/Programs/UnrealVS/BatchBuilderToolControl.xaml.cs
deleted file mode 100644
diff --git a/Engine/Source/Programs/UnrealVS/BatchBuilderToolWindow.cs b/Engine/Source/Programs/UnrealVS/BatchBuilderToolWindow.cs
deleted file mode 100644
diff --git a/Engine/Source/Programs/UnrealVS/BuildStartupProject.cs b/Engine/Source/Programs/UnrealVS/BuildStartupProject.cs
deleted file mode 100644
diff --git a/Engine/Source/Programs/UnrealVS/BuildStatusIcon.xaml b/Engine/Source/Programs/UnrealVS/BuildStatusIcon.xaml
deleted file mode 100644
diff --git a/Engine/Source/Programs/UnrealVS/BuildStatusIcon.xaml.cs b/Engine/Source/Programs/UnrealVS/BuildStatusIcon.xaml.cs
deleted file mode 100644
diff --git a/Engine/Source/Programs/UnrealVS/CommandLineEditor.cs b/Engine/Source/Programs/UnrealVS/CommandLineEditor.cs
deleted file mode 100644
diff --git a/Engine/Source/Programs/UnrealVS/CompileSingleFile.cs b/Engine/Source/Programs/UnrealVS/CompileSingleFile.cs
deleted file mode 100644
diff --git a/Engine/Source/Programs/UnrealVS/DeployToVS11Exp.bat b/Engine/Source/Programs/UnrealVS/DeployToVS11Exp.bat
deleted file mode 100755
diff --git a/Engine/Source/Programs/UnrealVS/GenerateProjectFiles.cs b/Engine/Source/Programs/UnrealVS/GenerateProjectFiles.cs
deleted file mode 100644
diff --git a/Engine/Source/Programs/UnrealVS/GlobalSuppressions.cs b/Engine/Source/Programs/UnrealVS/GlobalSuppressions.cs
deleted file mode 100644
diff --git a/Engine/Source/Programs/UnrealVS/Guids.cs b/Engine/Source/Programs/UnrealVS/Guids.cs
deleted file mode 100644
diff --git a/Engine/Source/Programs/UnrealVS/Logging.cs b/Engine/Source/Programs/UnrealVS/Logging.cs
deleted file mode 100644
diff --git a/Engine/Source/Programs/UnrealVS/ProfileManager.cs b/Engine/Source/Programs/UnrealVS/ProfileManager.cs
deleted file mode 100644
diff --git a/Engine/Source/Programs/UnrealVS/Properties/AssemblyInfo.cs b/Engine/Source/Programs/UnrealVS/Properties/AssemblyInfo.cs
deleted file mode 100644
diff --git a/Engine/Source/Programs/UnrealVS/ProvideSolutionProperties.cs b/Engine/Source/Programs/UnrealVS/ProvideSolutionProperties.cs
deleted file mode 100644
diff --git a/Engine/Source/Programs/UnrealVS/QuickBuild.cs b/Engine/Source/Programs/UnrealVS/QuickBuild.cs
deleted file mode 100644
diff --git a/Engine/Source/Programs/UnrealVS/StartupProjectSelector.cs b/Engine/Source/Programs/UnrealVS/StartupProjectSelector.cs
deleted file mode 100644
diff --git a/Engine/Source/Programs/UnrealVS/UnrealVS.2019/BatchBuilder.cs b/Engine/Source/Programs/UnrealVS/UnrealVS.2019/BatchBuilder.cs
new file mode 100644
--- /dev/null
+++ b/Engine/Source/Programs/UnrealVS/UnrealVS.2019/BatchBuilder.cs
@@ -0,0 +1,92 @@
+// Copyright Epic Games, Inc. All Rights Reserved.
+
+using Microsoft.VisualStudio.Shell;
+using Microsoft.VisualStudio.Shell.Interop;
+using System;
+using System.ComponentModel.Design;
+using System.IO;
+
+namespace UnrealVS
+{
+	internal class BatchBuilder
+	{
+		/** constants */
+
+		private const int BatchBuilderToolWindowId = 0x1300;
+
+		/** properties */
+
+		public bool IsBusy
+		{
+			get
+			{
+				return ToolControl.IsBusy;
+			}
+		}
+
+		public static BatchBuilderToolControl.BatchBuilderToolState _ToolControlState = new BatchBuilderToolControl.BatchBuilderToolState();
+		private static BatchBuilderToolControl _ToolControl;
+		public static BatchBuilderToolControl ToolControl
+		{
+			get
+			{
+				if (_ToolControl == null) _ToolControl = new BatchBuilderToolControl(_ToolControlState);
+				return _ToolControl;
+			}
+		}
+
+		/** methods */
+
+		public BatchBuilder()
+		{
+		}
+
+		public void Initialize()
+		{
+			// Create the command for the tool window
+			var ToolWindowCommandId = new CommandID(GuidList.UnrealVSCmdSet, BatchBuilderToolWindowId);
+			var ToolWindowMenuCommand = new MenuCommand(ShowToolWindow, ToolWindowCommandId);
+			UnrealVSPackage.Instance.MenuCommandService.AddCommand(ToolWindowMenuCommand);
+		}
+
+		/// <summary>
+		/// Called from the package class when there are options to be read out of the solution file.
+		/// </summary>
+		/// <param name="Stream">The stream to load the option data from.</param>
+		public void LoadOptions(Stream Stream)
+		{
+			_ToolControlState.LoadOptions(Stream);
+		}
+
+		/// <summary>
+		/// Called from the package class when there are options to be written to the solution file.
+		/// </summary>
+		/// <param name="Stream">The stream to save the option data to.</param>
+		public void SaveOptions(Stream Stream)
+		{
+			_ToolControlState.SaveOptions(Stream);
+		}
+
+		/// <summary>
+		/// Tick function called from the package
+		/// </summary>
+		public void Tick()
+		{
+			ToolControl.Tick();
+		}
+
+		private void ShowToolWindow(object sender, EventArgs e)
+		{
+			// Get the instance number 0 of this tool window. This window is single instance so this instance
+			// is actually the only one.
+			// The last flag is set to true so that if the tool window does not exists it will be created.
+			ToolWindowPane ToolWindow = UnrealVSPackage.Instance.FindToolWindow(typeof(BatchBuilderToolWindow), 0, true);
+			if ((null == ToolWindow) || (null == ToolWindow.Frame))
+			{
+				throw new NotSupportedException(Resources.ToolWindowCreateError);
+			}
+			IVsWindowFrame ToolWindowFrame = (IVsWindowFrame)ToolWindow.Frame;
+			Microsoft.VisualStudio.ErrorHandler.ThrowOnFailure(ToolWindowFrame.Show());
+		}
+	}
+}
diff --git a/Engine/Source/Programs/UnrealVS/UnrealVS.2019/BatchBuilderToolControl.xaml b/Engine/Source/Programs/UnrealVS/UnrealVS.2019/BatchBuilderToolControl.xaml
new file mode 100644
--- /dev/null
+++ b/Engine/Source/Programs/UnrealVS/UnrealVS.2019/BatchBuilderToolControl.xaml
@@ -0,0 +1,309 @@
+<!--Copyright Epic Games, Inc. All Rights Reserved.-->
+
+<UserControl x:Class="UnrealVS.BatchBuilderToolControl"
+			 xmlns="http://schemas.microsoft.com/winfx/2006/xaml/presentation"
+			 xmlns:x="http://schemas.microsoft.com/winfx/2006/xaml"
+			 xmlns:mc="http://schemas.openxmlformats.org/markup-compatibility/2006"
+			 xmlns:d="http://schemas.microsoft.com/expression/blend/2008"
+			 xmlns:vs_shell="clr-namespace:Microsoft.VisualStudio.PlatformUI;assembly=Microsoft.VisualStudio.Shell.15.0"
+			 xmlns:UnrealVS="clr-namespace:UnrealVS"
+			 mc:Ignorable="d" d:DesignHeight="250" d:DesignWidth="800"
+			 MinHeight="200" MinWidth="600" Name="ThisBatchBuilderToolCtrl"
+						 DataContext="{Binding RelativeSource={RelativeSource Self}}"
+						 IsEnabled="{Binding Path=IsSolutionOpen}">
+
+	<UserControl.CommandBindings>
+		<CommandBinding Command="Open" Executed="OpenItemOutputCommandExecuted"/>
+	</UserControl.CommandBindings>
+
+	<UserControl.Resources>
+
+		<!-- Style for all ListBoxItems -->
+		<Style TargetType="ListBoxItem">
+			<Setter Property="Template">
+				<Setter.Value>
+					<ControlTemplate TargetType="ListBoxItem">
+						<Border Name="Border">
+							<ContentPresenter />
+						</Border>
+						<ControlTemplate.Triggers>
+							<MultiTrigger>
+								<MultiTrigger.Conditions>
+									<Condition Property="IsSelected" Value="false"/>
+									<Condition Property="IsEnabled" Value="true"/>
+								</MultiTrigger.Conditions>
+								<Setter Property="Foreground" Value="{x:Static SystemColors.WindowTextBrush}" />
+								<Setter TargetName="Border" Property="Background" Value="{x:Static SystemColors.WindowBrush}" />
+							</MultiTrigger>
+							<Trigger Property="IsSelected" Value="true">
+								<Setter Property="Foreground" Value="{x:Static SystemColors.HighlightTextBrush}" />
+								<Setter TargetName="Border" Property="Background"
+										Value="{x:Static SystemColors.HighlightBrush}" />
+							</Trigger>
+						</ControlTemplate.Triggers>
+					</ControlTemplate>
+				</Setter.Value>
+			</Setter>
+		</Style>
+
+		<!-- DataTemplate ProjectListItemTemplate -->
+		<DataTemplate x:Key="ProjectListItemTemplate">
+			<TextBlock Text="{Binding}" />
+		</DataTemplate>
+
+		<!-- DataTemplate BuildListItemTemplate -->
+		<DataTemplate x:Key="BuildListItemTemplate">
+			<ContentControl MouseDoubleClick="OnDblClickBuildListItem" ToolTip="Double-click to switch the active solution config and the startup project to match this item.">
+				<TextBlock Text="{Binding Path=DisplayString}"/>
+			</ContentControl>
+		</DataTemplate>
+
+		<!-- DataTemplate BuildingListItemTemplate -->
+		<DataTemplate x:Key="BuildingListItemTemplate">
+			<ContentControl MouseDoubleClick="OnDblClickBuildingListItem">
+				<StackPanel Orientation="Horizontal">
+					<StackPanel.Style>
+						<Style>
+							<Style.Triggers>
+								<DataTrigger Binding="{Binding Path=HasOutputText}" Value="True">
+									<Setter Property="FrameworkElement.ContextMenu">
+										<Setter.Value>
+											<ContextMenu>
+												<MenuItem Header="Show Output" Command="Open" CommandParameter="{Binding}"/>
+											</ContextMenu>
+										</Setter.Value>
+									</Setter>
+								</DataTrigger>
+							</Style.Triggers>
+						</Style>
+					</StackPanel.Style>
+					<UnrealVS:BuildStatusIcon Margin="0,2"  Width="15" Height="15" BuildStatus="{Binding Path=JobStatus}" ToolTip="{Binding Path=JobStatusDisplayString}"/>
+					<TextBlock VerticalAlignment="Center" Margin="5,0,0,0" Text="{Binding Path=DisplayString}" ToolTip="{Binding Path=DisplayString}">
+						<TextBlock.Style>
+							<Style>
+								<Style.Triggers>
+									<DataTrigger Binding="{Binding Path=JobStatus}" Value="Executing">
+										<Setter Property="TextBlock.FontWeight" Value="Bold"/>
+									</DataTrigger>
+								</Style.Triggers>
+								<Setter Property="TextBlock.FontWeight" Value="Normal"/>
+							</Style>
+						</TextBlock.Style>
+					</TextBlock>
+				</StackPanel>
+			</ContentControl>
+		</DataTemplate>
+
+		<!-- Style DisabledOnBusy -->
+		<Style x:Key="DisabledOnBusy">
+			<Style.Triggers>
+				<DataTrigger Binding="{Binding Path=IsBusy}" Value="True">
+					<Setter Property="UIElement.IsEnabled" Value="False"/>
+				</DataTrigger>
+			</Style.Triggers>
+			<Setter Property="UIElement.IsEnabled" Value="True"/>
+		</Style>
+
+		<!-- Style ShowOnBusy -->
+		<Style x:Key="ShowOnBusy">
+			<Style.Triggers>
+				<DataTrigger Binding="{Binding Path=IsBusy}" Value="False">
+					<Setter Property="UIElement.Visibility" Value="Collapsed"/>
+				</DataTrigger>
+			</Style.Triggers>
+			<Setter Property="UIElement.Visibility" Value="Visible"/>
+		</Style>
+
+		<!-- Style BuildJobsListStyle -->
+		<Style x:Key="BuildJobsListStyle">
+			<Style.Triggers>
+				<DataTrigger Binding="{Binding Path=IsBusy}" Value="True">
+					<!--<Setter Property="ListBox.ItemTemplate" Value="{StaticResource NullItemTemplate}"/>-->
+					<Setter Property="ListBox.IsEnabled" Value="False"/>
+				</DataTrigger>
+			</Style.Triggers>
+			<Setter Property="ListBox.IsEnabled" Value="True"/>
+			<Setter Property="ListBox.ItemTemplate" Value="{DynamicResource BuildListItemTemplate}"/>
+			<Setter Property="ListBox.SelectionMode" Value="Extended"/>
+			<EventSetter Event="ListBox.SelectionChanged" Handler="OnBuildJobsSelectionChanged"/>
+		</Style>
+
+		<!-- Style BuildingJobsListStyle -->
+		<Style x:Key="BuildingJobsListStyle">
+			<Setter Property="ItemsControl.Background" Value="Transparent"/>
+			<Setter Property="ItemsControl.ItemTemplate" Value="{DynamicResource BuildingListItemTemplate}"/>
+			<Setter Property="ItemsControl.Margin" Value="3"/>
+		</Style>
+
+		<!-- Style OutputTabStyle -->
+		<Style x:Key="OutputTabStyle">
+			<Style.Triggers>
+				<DataTrigger Binding="{Binding Path=HasOutput}" Value="True">
+					<Setter Property="TabItem.Visibility" Value="Visible"/>
+				</DataTrigger>
+				<DataTrigger Binding="{Binding Path=HasOutput}" Value="False">
+					<Setter Property="TabItem.Visibility" Value="Collapsed"/>
+				</DataTrigger>
+			</Style.Triggers>
+		</Style>
+
+		<!-- Style StartButtonStyle -->
+		<Style x:Key="StartButtonStyle" TargetType="{x:Type Button}">
+			<Style.Triggers>
+				<DataTrigger Binding="{Binding Path=IsBusy}" Value="True">
+					<Setter Property="Content" Value="Stop"/>
+				</DataTrigger>
+			</Style.Triggers>
+			<Setter Property="Content" Value="Start"/>
+			<Setter Property="Margin" Value="0, 5, 0, 0"/>
+			<Setter Property="Width" Value="45"/>
+			<Setter Property="ToolTip" Value="Start/Stop build jobs in the list"/>
+			<EventSetter Event="Click" Handler="OnStartStopButtonClick"/>
+		</Style>
+
+		<!-- Style StopButtonStyle -->
+		<Style x:Key="StopButtonStyle" TargetType="{x:Type Button}">
+			<Style.Triggers>
+				<DataTrigger Binding="{Binding Path=IsBusy}" Value="False">
+					<Setter Property="IsEnabled" Value="False"/>
+				</DataTrigger>
+			</Style.Triggers>
+			<Setter Property="IsEnabled" Value="True"/>
+			<Setter Property="Content" Value="Stop"/>
+			<Setter Property="Margin" Value="0, 5, 0, 0"/>
+			<Setter Property="Width" Value="45"/>
+			<Setter Property="ToolTip" Value="Stop/cancel running build jobs"/>
+			<EventSetter Event="Click" Handler="OnStartStopButtonClick"/>
+		</Style>
+
+		<Style x:Key="LabelStyle" TargetType="{x:Type TextBlock}">
+			<Setter Property="Foreground" Value="{DynamicResource {x:Static vs_shell:EnvironmentColors.ToolWindowTextBrushKey}}"/>
+			<Setter Property="Background" Value="{DynamicResource {x:Static vs_shell:EnvironmentColors.ToolWindowBackgroundBrushKey}}"/>
+		</Style>
+
+	</UserControl.Resources>
+
+	<!-- CONTENT -->
+
+	<Grid Margin="3">
+		<Grid.ColumnDefinitions>
+			<ColumnDefinition Width="4*"/>
+			<ColumnDefinition Width="6"/>
+			<ColumnDefinition Width="3*"/>
+			<ColumnDefinition Width="6"/>
+			<ColumnDefinition Width="3*"/>
+			<ColumnDefinition Width="6"/>
+			<ColumnDefinition Width="7*"/>
+		</Grid.ColumnDefinitions>
+		<DockPanel Margin="0,3" Grid.Column="0">
+			<TextBlock Style="{StaticResource LabelStyle}" Margin="3,0,3,3" DockPanel.Dock="Top">Project</TextBlock>
+			<ListBox Name="ProjectsList" ItemsSource="{Binding Path=Projects}"
+					 SelectionMode="Extended" ItemTemplate="{DynamicResource ProjectListItemTemplate}" />
+		</DockPanel>
+		<GridSplitter Grid.Column="1" HorizontalAlignment="Stretch" Margin="1,22,1,3"/>
+		<DockPanel Margin="0,3" Grid.Column="2">
+			<TextBlock Style="{StaticResource LabelStyle}" Margin="3,0,3,3" DockPanel.Dock="Top">Config</TextBlock>
+			<ListBox Name="ConfigsList" ItemsSource="{Binding Path=Configs}"
+					 SelectionMode="Extended" />
+		</DockPanel>
+		<GridSplitter Grid.Column="3" HorizontalAlignment="Stretch" Margin="1,22,1,3"/>
+		<Image Margin="20,0,0,-28" Height="150" Width="150" HorizontalAlignment="Left" VerticalAlignment="Bottom" Grid.Column="4" Grid.ColumnSpan="3" Source="Resources/BatchBuilderBackground.png">
+			<Image.BitmapEffect>
+				<BlurBitmapEffect Radius="7.5"/>
+			</Image.BitmapEffect>
+		</Image>
+		<DockPanel Margin="0,3" Grid.Column="4">
+			<TextBlock Style="{StaticResource LabelStyle}" Margin="3,0,3,3" DockPanel.Dock="Top">Platform</TextBlock>
+			<StackPanel DockPanel.Dock="Bottom" HorizontalAlignment="Stretch" Background="#70ffffff" Margin="3,3,3,0">
+				<RadioButton Margin="3" Name="BuildRadioButton">Build</RadioButton>
+				<RadioButton Margin="3" Name="RebuildRadioButton">Rebuild</RadioButton>
+				<RadioButton Margin="3" Name="CleanRadioButton">Clean</RadioButton>
+			</StackPanel>
+			<ListBox Name="PlatformsList" ItemsSource="{Binding Path=Platforms}"
+					 SelectionMode="Extended" />
+		</DockPanel>
+		<GridSplitter Grid.Column="5" HorizontalAlignment="Stretch" Margin="1,22,1,3"/>
+		<DockPanel Grid.Column="6">
+			<StackPanel DockPanel.Dock="Left" VerticalAlignment="Center" Grid.Column="6" Background="#70ffffff" Margin="2">
+				<Button Style="{StaticResource DisabledOnBusy}" Margin="2,2,2,0" Width="24" Height="24" FontSize="15" Click="AddButtonClick" ToolTip="Adds Selected Projects, Configs and Platforms to the Build Jobs List">&gt;</Button>
+				<Button Style="{StaticResource DisabledOnBusy}" Margin="2" Width="24" Height="24" FontSize="15" Click="RemoveButtonClick" ToolTip="Removes Selected Build Jobs from the Build Jobs List">&lt;</Button>
+			</StackPanel>
+			<TabControl Margin="0,0,3,3" Grid.Column="7">
+				<TabItem Name="JobsListTab" Header="{Binding BuildJobsPanelTitle}">
+					<Grid>
+						<Grid.ColumnDefinitions>
+							<ColumnDefinition Width="*"/>
+							<ColumnDefinition Width="28"/>
+						</Grid.ColumnDefinitions>
+
+						<Grid.RowDefinitions>
+							<RowDefinition Height="*"/>
+							<RowDefinition Height="Auto"/>
+						</Grid.RowDefinitions>
+
+						<StackPanel Grid.Row="0" Margin="2,0" Grid.Column="1" Style="{StaticResource DisabledOnBusy}" HorizontalAlignment="Stretch">
+							<Button Margin="0" FontSize="15" Name="UpButton" ToolTip="Move Up" Click="UpButtonClick" IsEnabled="{Binding Path=IsSingleBuildJobSelected}">&#x2191;</Button>
+							<Button Margin="0,5,0,0" FontSize="15" Name="DownButton" ToolTip="Move Down" Click="DownButtonClick" IsEnabled="{Binding Path=IsSingleBuildJobSelected}">&#x2193;</Button>
+						</StackPanel>
+
+						<Grid Grid.Column="0" Grid.ColumnSpan="2" Grid.Row="1">
+							<Grid.ColumnDefinitions>
+								<ColumnDefinition Width="*"/>
+								<ColumnDefinition Width="Auto"/>
+							</Grid.ColumnDefinitions>
+
+							<GroupBox Grid.Column="0" Style="{StaticResource DisabledOnBusy}" Header="Job Set">
+								<DockPanel>
+									<Button DockPanel.Dock="Right" Margin="5,0,0,0" Height="22" Name="DeleteButton" Click="OnDeleteButtonClick" IsEnabled="{Binding IsDeletableSetSelected}" ToolTip="Deletes the selected Job Set">Delete</Button>
+									<ComboBox Name="SetCombo" IsEditable="True" Text="{Binding Path=SelectedBuildJobSetName, Mode=OneWayToSource, UpdateSourceTrigger=LostFocus}" ItemsSource="{Binding Path=BuildJobSets}" SelectionChanged="OnSetComboSelectionChanged" KeyDown="OnSetComboKeyDown"/>
+								</DockPanel>
+							</GroupBox>
+
+							<Grid Grid.Column="1" Margin="3,0,0,0">
+								<Grid.RowDefinitions>
+									<RowDefinition Height="13"/>
+									<RowDefinition Height="*"/>
+								</Grid.RowDefinitions>
+
+								<ProgressBar Grid.Row="0"
+										 Style="{StaticResource ShowOnBusy}" Margin="2" 
+										 IsIndeterminate="True" Foreground="#FF007ACC"
+										 Width="{Binding ActualWidth, ElementName=StartStopButton}"/>
+								<Button Grid.Row="1" Name="StartStopButton" Height="22"
+									VerticalAlignment="Top" 
+									Style="{StaticResource StartButtonStyle}"/>
+							</Grid>
+
+						</Grid>
+
+						<ListBox Name="BuildJobsList" Style="{StaticResource BuildJobsListStyle}">
+							<ListBox.ItemsSource>
+								<Binding Path="BuildJobs" />
+							</ListBox.ItemsSource>
+						</ListBox>
+					</Grid>
+				</TabItem>
+				<TabItem Name="OutputTab" Header="{Binding OutputPanelTitle}" Style="{StaticResource OutputTabStyle}">
+					<DockPanel>
+						<DockPanel DockPanel.Dock="Bottom">
+							<StackPanel Grid.Column="1" HorizontalAlignment="Right">
+								<ProgressBar Style="{StaticResource ShowOnBusy}" Height="10" Margin="0, 5, 0, 0" IsIndeterminate="True" Foreground="#FF007ACC"/>
+								<Button Name="StopButton" Style="{StaticResource StopButtonStyle}"/>
+							</StackPanel>
+						</DockPanel>
+						<ScrollViewer VerticalScrollBarVisibility="Auto">
+							<ItemsControl Name="BuildingJobsList" Style="{StaticResource BuildingJobsListStyle}">
+								<ItemsControl.ItemsSource>
+									<Binding Path="LastBuildJobsQueued" />
+								</ItemsControl.ItemsSource>
+							</ItemsControl>
+						</ScrollViewer>
+					</DockPanel>
+				</TabItem>
+			</TabControl>
+		</DockPanel>
+	</Grid>
+
+	<!-- END CONTENT -->
+
+</UserControl>
\ No newline at end of file
diff --git a/Engine/Source/Programs/UnrealVS/UnrealVS.2019/BatchBuilderToolControl.xaml.cs b/Engine/Source/Programs/UnrealVS/UnrealVS.2019/BatchBuilderToolControl.xaml.cs
new file mode 100644
--- /dev/null
+++ b/Engine/Source/Programs/UnrealVS/UnrealVS.2019/BatchBuilderToolControl.xaml.cs
@@ -0,0 +1,1050 @@
+// Copyright Epic Games, Inc. All Rights Reserved.
+
+using EnvDTE;
+using EnvDTE80;
+using Microsoft.VisualStudio;
+using Microsoft.VisualStudio.Shell;
+using Microsoft.VisualStudio.Shell.Interop;
+using System;
+using System.Collections.Generic;
+using System.Collections.ObjectModel;
+using System.IO;
+using System.Linq;
+using System.Text;
+using System.Windows;
+using System.Windows.Controls;
+using System.Windows.Input;
+using Constants = EnvDTE.Constants;
+
+namespace UnrealVS
+{
+	/// <summary>
+	/// Interaction logic for BatchBuilderToolControl.xaml
+	/// </summary>
+	public partial class BatchBuilderToolControl : UserControl
+	{
+		private readonly BatchBuilderToolState _state;
+		////////////// Types ///////////////////////
+
+		public class BuildJob : DependencyObject
+		{
+			////////////// BuildJob Types ///////////////////////
+
+			public enum BuildJobType
+			{
+				Build,
+				Rebuild,
+				Clean
+			}
+
+			public enum BuildJobStatus
+			{
+				Pending,
+				Executing,
+				Failed,
+				FailedToStart,
+				Cancelled,
+				Complete
+			}
+
+			////////////// BuildJob Serialized fields/properties ///////////////////////
+
+			private readonly Utils.SafeProjectReference _Project;
+			private readonly string _Config;
+			private readonly string _Platform;
+			private readonly BuildJobType _JobType;
+
+			////////////// BuildJob Properties ///////////////////////
+
+			public Utils.SafeProjectReference Project
+			{
+				get { return _Project; }
+			}
+
+			public string Config
+			{
+				get { return _Config; }
+			}
+
+			public string Platform
+			{
+				get { return _Platform; }
+			}
+
+			public BuildJobType JobType
+			{
+				get { return _JobType; }
+			}
+
+			public TimeSpan ExecutionTime { get; set; }
+
+			private readonly string _SolutionConfigString;
+			public string SolutionConfigString
+			{
+				get { return _SolutionConfigString; }
+			}
+
+			private readonly string _DisplayString;
+			public string DisplayString
+			{
+				get { return _DisplayString; }
+			}
+
+			////////////// BuildJob Dependency properties ///////////////////////
+
+			public static readonly DependencyProperty JobStatusProperty;
+			public static readonly DependencyProperty JobStatusDisplayStringProperty;
+			public static readonly DependencyProperty OutputTextProperty;
+			public static readonly DependencyProperty HasOutputTextProperty;
+
+			////////////// BuildJob Dependency property wrappers ///////////////////////
+
+			public BuildJobStatus JobStatus
+			{
+				get { return (BuildJobStatus)GetValue(JobStatusProperty); }
+				set { SetValue(JobStatusProperty, value); }
+			}
+
+			public string JobStatusDisplayString
+			{
+				get { return (string)GetValue(JobStatusDisplayStringProperty); }
+				private set { SetValue(JobStatusDisplayStringProperty, value); }
+			}
+
+			public string OutputText
+			{
+				get { return (string)GetValue(OutputTextProperty); }
+				set { SetValue(OutputTextProperty, value); }
+			}
+
+			public bool HasOutputText
+			{
+				get { return (bool)GetValue(HasOutputTextProperty); }
+				private set { SetValue(HasOutputTextProperty, value); }
+			}
+
+			////////////// BuildJob static methods ///////////////////////
+
+			static BuildJob()
+			{
+				// Register the dependency properties
+				JobStatusProperty = DependencyProperty.Register("JobStatus",
+																typeof(BuildJobStatus), typeof(BuildJob),
+																new FrameworkPropertyMetadata(BuildJobStatus.Pending,
+																							  OnJobStatusChanged));
+
+				JobStatusDisplayStringProperty = DependencyProperty.Register("JobStatusDisplayString",
+																			 typeof(string), typeof(BuildJob),
+																			 new FrameworkPropertyMetadata(InvalidJobStatusString));
+
+				OutputTextProperty = DependencyProperty.Register("OutputText",
+																 typeof(string), typeof(BuildJob),
+																 new FrameworkPropertyMetadata(String.Empty, OnOutputTextChanged));
+
+				HasOutputTextProperty = DependencyProperty.Register("HasOutputText",
+																	typeof(bool), typeof(BuildJob),
+																	new FrameworkPropertyMetadata(false));
+			}
+
+			private static void OnJobStatusChanged(DependencyObject Obj, DependencyPropertyChangedEventArgs Args)
+			{
+				BuildJob ThisBuildJob = (BuildJob)Obj;
+				if (ThisBuildJob != null)
+				{
+					ThisBuildJob.JobStatusDisplayString = ThisBuildJob.GetJobStatusString();
+				}
+			}
+
+			private static void OnOutputTextChanged(DependencyObject Obj, DependencyPropertyChangedEventArgs Args)
+			{
+				BuildJob ThisBuildJob = (BuildJob)Obj;
+				if (ThisBuildJob != null)
+				{
+					ThisBuildJob.HasOutputText = !String.IsNullOrEmpty(ThisBuildJob.OutputText);
+				}
+			}
+
+			////////////// BuildJob methods ///////////////////////
+
+			public BuildJob(Utils.SafeProjectReference InProject, string InConfig, string InPlatform, BuildJobType InJobType)
+			{
+				_Project = InProject;
+				_Config = InConfig;
+				_Platform = InPlatform;
+				_JobType = InJobType;
+
+				_SolutionConfigString = Config + '|' + Platform;
+				_DisplayString = GetDisplayString();
+				JobStatusDisplayString = GetJobStatusString();
+			}
+
+			private string GetDisplayString()
+			{
+				if (Project == null || Config == null || Platform == null) return InvalidDisplayString;
+
+				return String.Format("{0}: {1} [{2}]", Enum.GetName(typeof(BuildJobType), JobType), Project.Name,
+									 SolutionConfigString);
+			}
+
+			private string GetJobStatusString()
+			{
+				return Enum.GetName(typeof(BuildJobStatus), JobStatus);
+			}
+
+			////////////// BuildJob fields ///////////////////////
+
+			private const string InvalidDisplayString = "INVALIDBUILDJOB";
+			private const string InvalidJobStatusString = "INVALIDSTATUS";
+		}
+
+		public class BuildJobSet
+		{
+			private readonly ObservableCollection<BuildJob> _BuildJobs = new ObservableCollection<BuildJob>();
+
+			public string Name { get; set; }
+
+			public ObservableCollection<BuildJob> BuildJobs
+			{
+				get { return _BuildJobs; }
+			}
+
+			public void DeepCopyJobsFrom(ObservableCollection<BuildJob> From)
+			{
+				foreach (var Job in From)
+				{
+					BuildJobs.Add(new BuildJob(Job.Project, Job.Config, Job.Platform, Job.JobType));
+				}
+			}
+
+			public override string ToString()
+			{
+				return Name;
+			}
+		}
+
+		public struct ProjectListItem
+		{
+			public Project Project { get; set; }
+
+			public override string ToString()
+			{
+				if (Project != null) return Project.Name;
+				return "<None>";
+			}
+		}
+
+		public class BatchBuilderToolState
+		{
+			public readonly ObservableCollection<BuildJobSet> _BuildJobSetsCollection = new ObservableCollection<BuildJobSet>();
+
+			/// <summary>
+			/// Called from the package class when there are options to be read out of the solution file.
+			/// </summary>
+			/// <param name="Stream">The stream to load the option data from.</param>
+			public void LoadOptions(Stream Stream)
+			{
+				try
+				{
+					_BuildJobSetsCollection.Clear();
+
+					using (BinaryReader Reader = new BinaryReader(Stream))
+					{
+						int SetCount = Reader.ReadInt32();
+
+						for (int SetIdx = 0; SetIdx < SetCount; SetIdx++)
+						{
+							BuildJobSet LoadedSet = new BuildJobSet
+							{
+								Name = Reader.ReadString()
+							};
+							int JobCount = Reader.ReadInt32();
+							for (int JobIdx = 0; JobIdx < JobCount; JobIdx++)
+							{
+								Utils.SafeProjectReference ProjectRef = new Utils.SafeProjectReference { FullName = Reader.ReadString(), Name = Reader.ReadString() };
+
+								string Config = Reader.ReadString();
+								string Platform = Reader.ReadString();
+
+								if (Enum.TryParse(Reader.ReadString(), out BuildJob.BuildJobType JobType))
+								{
+									LoadedSet.BuildJobs.Add(new BuildJob(ProjectRef, Config, Platform, JobType));
+								}
+							}
+							_BuildJobSetsCollection.Add(LoadedSet);
+						}
+					}
+
+					if (_BuildJobSetsCollection.Count == 0)
+					{
+						BuildJobSet DefaultItem = new BuildJobSet { Name = "UntitledSet" };
+						_BuildJobSetsCollection.Add(DefaultItem);
+					}
+
+					StateLoaded?.Invoke(this, null);
+				}
+				catch (Exception ex)
+				{
+					Exception AppEx = new ApplicationException("BatchBuilder failed to load options from .suo", ex);
+					Logging.WriteLine(AppEx.ToString());
+					throw AppEx;
+				}
+			}
+
+			/// <summary>
+			/// Called from the package class when there are options to be written to the solution file.
+			/// </summary>
+			/// <param name="Stream">The stream to save the option data to.</param>
+			public void SaveOptions(Stream Stream)
+			{
+				try
+				{
+					using (BinaryWriter Writer = new BinaryWriter(Stream))
+					{
+						Writer.Write(_BuildJobSetsCollection.Count);
+						foreach (var Set in _BuildJobSetsCollection)
+						{
+							Writer.Write(Set.Name);
+							Writer.Write(Set.BuildJobs.Count);
+							foreach (var Job in Set.BuildJobs)
+							{
+								Writer.Write(Job.Project.FullName);
+								Writer.Write(Job.Project.Name);
+								Writer.Write(Job.Config);
+								Writer.Write(Job.Platform);
+								Writer.Write(Enum.GetName(typeof(BuildJob.BuildJobType), Job.JobType) ?? "INVALIDJOBTYPE");
+							}
+						}
+					}
+				}
+				catch (Exception ex)
+				{
+					Exception AppEx = new ApplicationException("BatchBuilder failed to save options to .suo", ex);
+					Logging.WriteLine(AppEx.ToString());
+					throw AppEx;
+				}
+			}
+
+			// we are using a custom even instead of listening to the observable collection as we just want a even when the entire collection has been modified
+			public event EventHandler StateLoaded;
+		}
+
+		////////////// Dependency properties ///////////////////////
+
+		public static readonly DependencyProperty IsSolutionOpenProperty;
+		public static readonly DependencyProperty IsSingleBuildJobSelectedProperty;
+		public static readonly DependencyProperty BuildJobsProperty;
+		public static readonly DependencyProperty IsDeletableSetSelectedProperty;
+		public static readonly DependencyProperty BuildJobsPanelTitleProperty;
+		public static readonly DependencyProperty OutputPanelTitleProperty;
+		public static readonly DependencyProperty HasOutputProperty;
+		public static readonly DependencyProperty IsBusyProperty;
+
+		////////////// Dependency property wrappers ///////////////////////
+
+		public bool IsSolutionOpen
+		{
+			get { return (bool)GetValue(IsSolutionOpenProperty); }
+			set { SetValue(IsSolutionOpenProperty, value); }
+		}
+
+		public bool IsSingleBuildJobSelected
+		{
+			get { return (bool)GetValue(IsSingleBuildJobSelectedProperty); }
+			set { SetValue(IsSingleBuildJobSelectedProperty, value); }
+		}
+
+		public ObservableCollection<BuildJob> BuildJobs
+		{
+			get { return (ObservableCollection<BuildJob>)GetValue(BuildJobsProperty); }
+			set { SetValue(BuildJobsProperty, value); }
+		}
+
+		private readonly BuildJobSet _LastBuildJobsQueued = new BuildJobSet();
+		public ObservableCollection<BuildJob> LastBuildJobsQueued
+		{
+			get { return _LastBuildJobsQueued.BuildJobs; }
+		}
+
+		public bool IsDeletableSetSelected
+		{
+			get { return (bool)GetValue(IsDeletableSetSelectedProperty); }
+			set { SetValue(IsDeletableSetSelectedProperty, value); }
+		}
+
+		public string BuildJobsPanelTitle
+		{
+			get { return (string)GetValue(BuildJobsPanelTitleProperty); }
+			set { SetValue(BuildJobsPanelTitleProperty, value); }
+		}
+
+		public string OutputPanelTitle
+		{
+			get { return (string)GetValue(OutputPanelTitleProperty); }
+			set { SetValue(OutputPanelTitleProperty, value); }
+		}
+
+		public bool HasOutput
+		{
+			get { return (bool)GetValue(HasOutputProperty); }
+			set { SetValue(HasOutputProperty, value); }
+		}
+
+		public bool IsBusy
+		{
+			get { return (bool)GetValue(IsBusyProperty); }
+			set { SetValue(IsBusyProperty, value); }
+		}
+
+		////////////// 'Normal' Properties ///////////////////////
+
+		public ObservableCollection<ProjectListItem> Projects
+		{
+			get { return _ProjectCollection; }
+		}
+
+		public ObservableCollection<string> Configs
+		{
+			get { return _ConfigCollection; }
+		}
+
+		public ObservableCollection<string> Platforms
+		{
+			get { return _PlatformCollection; }
+		}
+
+		public ObservableCollection<BuildJobSet> BuildJobSets
+		{
+			get { return _state._BuildJobSetsCollection; }
+		}
+
+		public string SelectedBuildJobSetName
+		{
+			set
+			{
+				if (String.IsNullOrEmpty(value)) return;
+
+				BuildJobSet SelectedSet =
+					_state._BuildJobSetsCollection.FirstOrDefault(
+						Set => String.Compare(Set.Name, value, StringComparison.InvariantCultureIgnoreCase) == 0);
+
+				if (SelectedSet == null)
+				{
+					SelectedSet = new BuildJobSet { Name = value };
+					SelectedSet.DeepCopyJobsFrom(BuildJobs);
+					_state._BuildJobSetsCollection.Add(SelectedSet);
+				}
+
+				SetCombo.SelectedItem = SelectedSet;
+			}
+		}
+
+		////////////// static methods ///////////////////////
+
+		/// <summary>
+		/// Static consturctor initializes dependency properties
+		/// </summary>
+		static BatchBuilderToolControl()
+		{
+			// Register the dependency properties
+			IsSolutionOpenProperty = DependencyProperty.Register("IsSolutionOpen",
+																 typeof(bool), typeof(BatchBuilderToolControl),
+																 new FrameworkPropertyMetadata(false));
+
+			IsSingleBuildJobSelectedProperty = DependencyProperty.Register("IsSingleBuildJobSelected",
+																		   typeof(bool), typeof(BatchBuilderToolControl),
+																		   new FrameworkPropertyMetadata(false));
+
+			BuildJobsProperty = DependencyProperty.Register("BuildJobs",
+															typeof(ObservableCollection<BuildJob>),
+															typeof(BatchBuilderToolControl),
+															new FrameworkPropertyMetadata(new ObservableCollection<BuildJob>()));
+
+			IsDeletableSetSelectedProperty = DependencyProperty.Register("IsDeletableSetSelected",
+																		 typeof(bool), typeof(BatchBuilderToolControl),
+																		 new FrameworkPropertyMetadata(false));
+
+			BuildJobsPanelTitleProperty = DependencyProperty.Register("BuildJobsPanelTitle",
+																	  typeof(string), typeof(BatchBuilderToolControl),
+																	  new FrameworkPropertyMetadata(BuildJobsPanelPrefix));
+
+			OutputPanelTitleProperty = DependencyProperty.Register("OutputPanelTitle",
+																   typeof(string), typeof(BatchBuilderToolControl),
+																   new FrameworkPropertyMetadata(OutputPanelTitlePrefix));
+
+			HasOutputProperty = DependencyProperty.Register("HasOutput",
+															typeof(bool), typeof(BatchBuilderToolControl),
+															new FrameworkPropertyMetadata(false));
+
+			IsBusyProperty = DependencyProperty.Register("IsBusy",
+														 typeof(bool), typeof(BatchBuilderToolControl),
+														 new FrameworkPropertyMetadata(false));
+		}
+
+		private static void DisplayBatchOutputText(string Text)
+		{
+			if (string.IsNullOrEmpty(Text)) return;
+
+			UnrealVSPackage.Instance.DTE.ExecuteCommand("View.Output");
+
+			var Pane = UnrealVSPackage.Instance.GetOutputPane(GuidList.BatchBuildPaneGuid, "UnrealVS - BatchBuild");
+			if (Pane != null)
+			{
+				// Clear and activate the output pane.
+				Pane.Clear();
+
+				Pane.OutputString(Text);
+
+				// @todo: Activating doesn't seem to really bring the pane to front like we would expect it to.
+				Pane.Activate();
+			}
+		}
+
+		////////////// public methods ///////////////////////
+
+		/// <summary>
+		/// Default constructor
+		/// </summary>
+		public BatchBuilderToolControl(BatchBuilderToolState state)
+		{
+			_state = state;
+			IsSolutionOpen = UnrealVSPackage.Instance.DTE.Solution.IsOpen;
+
+			InitializeComponent();
+
+			UnrealVSPackage.Instance.OnBuildBegin += OnBuildBegin;
+			UnrealVSPackage.Instance.OnBuildDone += OnBuildDone;
+
+			// Register for events that we care about
+			UnrealVSPackage.Instance.StartupProjectSelector.StartupProjectListChanged += RefreshProjectCollection;
+			UnrealVSPackage.Instance.OnSolutionOpened += delegate
+															 {
+																 IsSolutionOpen = true;
+																 RefreshConfigAndPlatformCollections();
+															 };
+			UnrealVSPackage.Instance.OnSolutionClosed += delegate
+															 {
+																 IsSolutionOpen = false;
+																 RefreshConfigAndPlatformCollections();
+															 };
+
+			RefreshProjectCollection(UnrealVSPackage.Instance.StartupProjectSelector.StartupProjectList);
+			RefreshConfigAndPlatformCollections();
+
+			BuildRadioButton.IsChecked = true;
+
+			EnsureDefaultBuildJobSet();
+
+			state.StateLoaded += BuildJobSetsCollectionOnCollectionChanged;
+		}
+
+		private void BuildJobSetsCollectionOnCollectionChanged(object sender, EventArgs eventArgs)
+		{
+			Logging.WriteLine("BuildJobSet Collection changed");
+
+			if (SetCombo.SelectedItem == null)
+			{
+				SetCombo.SelectedItem = _state._BuildJobSetsCollection[0];
+			}
+		}
+
+		/// <summary>
+		/// Tick function called from the package via BatchBuilder.Tick()
+		/// </summary>
+		public void Tick()
+		{
+			TickBuildQueue();
+		}
+
+		////////////// private methods ///////////////////////
+
+		private void EnsureDefaultBuildJobSet()
+		{
+			if (_state._BuildJobSetsCollection.Count == 0)
+			{
+				BuildJobSet DefaultItem = new BuildJobSet { Name = "UntitledSet" };
+				_state._BuildJobSetsCollection.Add(DefaultItem);
+				if (SetCombo != null)
+				{
+					SetCombo.SelectedItem = DefaultItem;
+				}
+			}
+		}
+
+		private void RefreshProjectCollection(object sender, StartupProjectSelector.StartupProjectListChangedEventArgs e)
+		{
+			RefreshProjectCollection(e.StartupProjects);
+		}
+
+		private void RefreshProjectCollection(Project[] StartupProjects)
+		{
+			_ProjectCollection.Clear();
+
+			foreach (var Project in StartupProjects)
+			{
+				_ProjectCollection.Add(new ProjectListItem { Project = Project });
+			}
+		}
+
+		private void RefreshConfigAndPlatformCollections()
+		{
+			_ConfigCollection.Clear();
+			_PlatformCollection.Clear();
+
+			if (!IsSolutionOpen)
+			{
+				_state._BuildJobSetsCollection.Clear();
+			}
+
+			Utils.GetSolutionConfigsAndPlatforms(out string[] RefreshConfigs, out string[] RefreshPlatforms);
+
+			foreach (string Config in RefreshConfigs)
+			{
+				_ConfigCollection.Add(Config);
+			}
+			foreach (string Platform in RefreshPlatforms)
+			{
+				_PlatformCollection.Add(Platform);
+			}
+
+			EnsureDefaultBuildJobSet();
+		}
+
+		private void AddButtonClick(object Sender, RoutedEventArgs E)
+		{
+			BuildJob[] Buildables = GetSelectedBuildableOnLeft();
+
+			if (Buildables.Length > 0)
+			{
+				foreach (var BuildJob in Buildables)
+				{
+					BuildJobs.Add(BuildJob);
+				}
+				JobsListTab.IsSelected = true;
+			}
+		}
+
+		private void RemoveButtonClick(object Sender, RoutedEventArgs E)
+		{
+			BuildJob[] Buildables = GetSelectedBuildableOnRight();
+
+			foreach (var BuildJob in Buildables)
+			{
+				BuildJobs.Remove(BuildJob);
+			}
+
+			JobsListTab.IsSelected = true;
+		}
+
+		private void UpButtonClick(object Sender, RoutedEventArgs E)
+		{
+			if (!IsSingleBuildJobSelected) return;
+
+			BuildJob SelectedJob = (BuildJob)BuildJobsList.SelectedItem;
+
+			int CurrentIdx = BuildJobs.IndexOf(SelectedJob);
+			if (CurrentIdx > 0)
+			{
+				BuildJobs.Move(CurrentIdx, CurrentIdx - 1);
+			}
+		}
+
+		private void DownButtonClick(object Sender, RoutedEventArgs E)
+		{
+			if (!IsSingleBuildJobSelected) return;
+
+			BuildJob SelectedJob = (BuildJob)BuildJobsList.SelectedItem;
+
+			int CurrentIdx = BuildJobs.IndexOf(SelectedJob);
+			if (CurrentIdx < BuildJobs.Count - 1)
+			{
+				BuildJobs.Move(CurrentIdx, CurrentIdx + 1);
+			}
+		}
+
+		private BuildJob[] GetSelectedBuildableOnLeft()
+		{
+			BuildJob.BuildJobType JobType = GetSelectedBuildJobType();
+
+			List<BuildJob> Buildables = new List<BuildJob>();
+
+			foreach (ProjectListItem Project in ProjectsList.SelectedItems)
+			{
+				foreach (string Config in ConfigsList.SelectedItems)
+				{
+					foreach (string Platform in PlatformsList.SelectedItems)
+					{
+						Buildables.Add(
+							new BuildJob(new Utils.SafeProjectReference { FullName = Project.Project.FullName, Name = Project.Project.Name },
+										 Config, Platform, JobType));
+					}
+				}
+			}
+
+			return Buildables.ToArray();
+		}
+
+		private BuildJob[] GetSelectedBuildableOnRight()
+		{
+			List<BuildJob> Buildables = new List<BuildJob>();
+
+			foreach (BuildJob BuildJob in BuildJobsList.SelectedItems)
+			{
+				Buildables.Add(BuildJob);
+			}
+
+			return Buildables.ToArray();
+		}
+
+		private BuildJob.BuildJobType GetSelectedBuildJobType()
+		{
+			if (BuildRadioButton.IsChecked.HasValue && BuildRadioButton.IsChecked.Value)
+			{
+				return BuildJob.BuildJobType.Build;
+			}
+			if (RebuildRadioButton.IsChecked.HasValue && RebuildRadioButton.IsChecked.Value)
+			{
+				return BuildJob.BuildJobType.Rebuild;
+			}
+			return BuildJob.BuildJobType.Clean;
+		}
+
+		private void OnBuildBegin(out int Cancel)
+		{
+			Cancel = 0;
+			if (IsBusy && _ActiveBuildJob == null)
+			{
+				Cancel = 1;
+			}
+		}
+
+		private void OnBuildDone(bool bSucceeded, bool bModified, bool bWasCancelled)
+		{
+			if (_ActiveBuildJob != null)
+			{
+				if (bWasCancelled)
+				{
+					_ActiveBuildJob.JobStatus = BuildJob.BuildJobStatus.Cancelled;
+					if (IsBusy)
+					{
+						foreach (var Job in _BuildQueue)
+						{
+							Job.JobStatus = BuildJob.BuildJobStatus.Cancelled;
+							Job.OutputText = GetBuildJobOutputText(Job, null);
+						}
+						_BuildQueue.Clear();
+					}
+				}
+				else if (bSucceeded)
+				{
+					_ActiveBuildJob.JobStatus = BuildJob.BuildJobStatus.Complete;
+				}
+				else
+				{
+					_ActiveBuildJob.JobStatus = BuildJob.BuildJobStatus.Failed;
+				}
+
+				_ActiveBuildJob.OutputText = GetBuildJobOutputText(_ActiveBuildJob, _BuildJobStartTime);
+			}
+
+			_ActiveBuildJob = null;
+			UpdateBusyState();
+		}
+
+		private void OnBuildJobsSelectionChanged(object sender, SelectionChangedEventArgs e)
+		{
+			IsSingleBuildJobSelected = BuildJobsList != null && BuildJobsList.SelectedItems.Count == 1;
+		}
+
+		private void OnSetComboSelectionChanged(object sender, SelectionChangedEventArgs e)
+		{
+			if (SetCombo.SelectedItem != null)
+			{
+				_LastSelectedBuildJobSet = (BuildJobSet)SetCombo.SelectedItem;
+				BuildJobs = _LastSelectedBuildJobSet.BuildJobs;
+				IsDeletableSetSelected = _state._BuildJobSetsCollection.Count > 1;
+				BuildJobsPanelTitle = String.Format("{0} ({1})", BuildJobsPanelPrefix, _LastSelectedBuildJobSet.Name);
+			}
+			if (!_state._BuildJobSetsCollection.Contains(_LastSelectedBuildJobSet))
+			{
+				_LastSelectedBuildJobSet = null;
+				BuildJobs = new ObservableCollection<BuildJob>();
+				IsDeletableSetSelected = false;
+				BuildJobsPanelTitle = BuildJobsPanelPrefix;
+			}
+		}
+
+		private void OnSetComboKeyDown(object sender, KeyEventArgs e)
+		{
+			if (e.Key == Key.Enter)
+			{
+				SelectedBuildJobSetName = SetCombo.Text;
+			}
+		}
+
+		private void OnDeleteButtonClick(object sender, RoutedEventArgs e)
+		{
+			if (_LastSelectedBuildJobSet != null)
+			{
+				int SelectionIdx = _state._BuildJobSetsCollection.IndexOf(_LastSelectedBuildJobSet);
+				if (SelectionIdx >= 0 && SelectionIdx < _state._BuildJobSetsCollection.Count)
+				{
+					SetCombo.SelectedItem = null;
+					_state._BuildJobSetsCollection.RemoveAt(SelectionIdx);
+					if (!(SelectionIdx < _state._BuildJobSetsCollection.Count))
+					{
+						SelectionIdx--;
+					}
+					if (SelectionIdx >= 0)
+					{
+						SetCombo.SelectedIndex = SelectionIdx;
+					}
+				}
+			}
+		}
+
+		private void OnStartStopButtonClick(object sender, RoutedEventArgs e)
+		{
+			if (!IsBusy)
+			{
+				// Start
+				RunBuildJobs();
+			}
+			else
+			{
+				// Stop
+				UnrealVSPackage.Instance.DTE.ExecuteCommand("Build.Cancel");
+				if (_ActiveBuildJob != null)
+				{
+					_ActiveBuildJob.JobStatus = BuildJob.BuildJobStatus.Cancelled;
+					_ActiveBuildJob.OutputText = GetBuildJobOutputText(_ActiveBuildJob, _BuildJobStartTime);
+					_ActiveBuildJob = null;
+				}
+				foreach (var BuildJob in _BuildQueue)
+				{
+					BuildJob.JobStatus = BuildJob.BuildJobStatus.Cancelled;
+					BuildJob.OutputText = GetBuildJobOutputText(BuildJob, null);
+				}
+				_BuildQueue.Clear();
+				UpdateBusyState();
+			}
+		}
+
+		private void RunBuildJobs()
+		{
+			if (BuildJobs == null)
+				throw new ArgumentNullException(nameof(BuildJobs));
+
+			if (_LastBuildJobsQueued == null)
+				throw new ArgumentNullException(nameof(_LastBuildJobsQueued));
+
+			int BuildJobCount = BuildJobs.Count;
+			SanityCheckBuildJobs();
+
+			if (BuildJobCount > BuildJobs.Count)
+			{
+				if (VSConstants.S_OK != VsShellUtilities.ShowMessageBox(ServiceProvider.GlobalProvider,
+																		"WARNING: Some build jobs were deleted because they were invalid.",
+																		"UnrealVS",
+																		OLEMSGICON.OLEMSGICON_WARNING,
+																		OLEMSGBUTTON.OLEMSGBUTTON_OKCANCEL,
+																		OLEMSGDEFBUTTON.OLEMSGDEFBUTTON_FIRST))
+				{
+					return;
+				}
+			}
+
+			_BuildQueue.Clear();
+			_LastBuildJobsQueued.BuildJobs.Clear();
+			_LastBuildJobsQueued.DeepCopyJobsFrom(BuildJobs);
+			foreach (var Job in _LastBuildJobsQueued.BuildJobs)
+			{
+				_BuildQueue.Enqueue(Job);
+			}
+
+			HasOutput = _BuildQueue.Count > 0;
+			if (_LastSelectedBuildJobSet != null)
+			{
+				OutputPanelTitle = String.Format("{0} ({1})", OutputPanelTitlePrefix, _LastSelectedBuildJobSet.Name);
+			}
+			OutputTab.IsSelected = HasOutput;
+
+			UpdateBusyState();
+		}
+
+		private void UpdateBusyState()
+		{
+			IsBusy = _BuildQueue.Count > 0 || _ActiveBuildJob != null;
+		}
+
+		private void SanityCheckBuildJobs(Project[] CheckProjects)
+		{
+			// check that project, config and platform strings in build jobs are valid
+			// remove invalid ones
+			Utils.GetSolutionConfigsAndPlatforms(out string[] SolutionConfigs, out string[] SolutionPlatforms);
+
+			foreach (var JobSet in _state._BuildJobSetsCollection)
+			{
+				for (int JobIdx = JobSet.BuildJobs.Count - 1; JobIdx >= 0; JobIdx--)
+				{
+					BuildJob Job = JobSet.BuildJobs[JobIdx];
+
+					bool bProjectOkay = CheckProjects.Any(Proj => String.CompareOrdinal(Job.Project.FullName, Proj.FullName) == 0);
+					bool bConfigOkay = SolutionConfigs.Any(Config => String.CompareOrdinal(Job.Config, Config) == 0);
+					bool bPlatformOkay = SolutionPlatforms.Any(Platform => String.CompareOrdinal(Job.Platform, Platform) == 0);
+
+					if (!bProjectOkay || !bConfigOkay || !bPlatformOkay)
+					{
+						JobSet.BuildJobs.RemoveAt(JobIdx);
+					}
+				}
+			}
+		}
+
+		private void SanityCheckBuildJobs()
+		{
+			SanityCheckBuildJobs(Utils.GetAllProjectsFromDTE());
+		}
+
+		private void TickBuildQueue()
+		{
+			if (_ActiveBuildJob != null || _BuildQueue.Count == 0) return;
+
+			while (_ActiveBuildJob == null && _BuildQueue.Count > 0)
+			{
+				BuildJob Job = _BuildQueue.Dequeue();
+
+				Project Project = Job.Project.GetProjectSlow();
+
+				if (Project != null)
+				{
+					Utils.ExecuteProjectBuild(
+						Project,
+						Job.Config,
+						Job.Platform,
+						Job.JobType,
+						delegate
+						{
+							// Job starting
+							_ActiveBuildJob = Job;
+							_ActiveBuildJob.JobStatus = BuildJob.BuildJobStatus.Executing;
+							_BuildJobStartTime = DateTime.Now;
+						},
+						delegate
+						{
+							// Job failed to start - clear active job
+							_ActiveBuildJob.JobStatus = BuildJob.BuildJobStatus.FailedToStart;
+							_ActiveBuildJob.OutputText = GetBuildJobOutputText(_ActiveBuildJob, _BuildJobStartTime);
+							_ActiveBuildJob = null;
+						});
+				}
+			}
+		}
+
+		private string GetBuildJobOutputText(BuildJob Job, DateTime? StartTime)
+		{
+			var OutputText = new StringBuilder();
+
+			var Window = UnrealVSPackage.Instance.DTE.Windows.Item(Constants.vsWindowKindOutput);
+			if (Window != null)
+			{
+				var OutputWindow = Window.Object as OutputWindow;
+				if (OutputWindow != null)
+				{
+					var OutputWindowPane = OutputWindow.OutputWindowPanes.Item("Build");
+					if (OutputWindowPane != null)
+					{
+						var TextStartPoint = (EditPoint2)OutputWindowPane.TextDocument.StartPoint.CreateEditPoint();
+						string BuildOutputText = TextStartPoint.GetText(OutputWindowPane.TextDocument.EndPoint);
+
+						OutputText.AppendLine("========== Build Output ==========");
+						OutputText.AppendLine(string.Empty);
+						OutputText.Append(BuildOutputText);
+						OutputText.AppendLine(string.Empty);
+					}
+				}
+			}
+
+			int JobIndex = LastBuildJobsQueued.IndexOf(Job) + 1;
+			int JobCount = LastBuildJobsQueued.Count;
+
+			OutputText.AppendLine(string.Format("========== Batch Builder {0} - Job {1} of {2} ==========", OutputPanelTitle,
+												JobIndex, JobCount));
+			OutputText.AppendLine(string.Empty);
+			OutputText.AppendLine("------ " + Job.DisplayString + " ------");
+			OutputText.AppendLine("Status: " + Job.JobStatusDisplayString);
+
+			if (StartTime.HasValue)
+			{
+				OutputText.AppendLine("Started: " + StartTime.Value.ToLongTimeString());
+
+				if (Job.JobStatus == BuildJob.BuildJobStatus.Complete ||
+					Job.JobStatus == BuildJob.BuildJobStatus.Failed ||
+					Job.JobStatus == BuildJob.BuildJobStatus.Cancelled)
+				{
+					DateTime EndTime = DateTime.Now;
+					OutputText.AppendLine("Ended: " + EndTime.ToLongTimeString());
+					OutputText.AppendLine(string.Format("Execution time: {0:F2} seconds", (EndTime - StartTime.Value).TotalSeconds));
+				}
+			}
+
+			return OutputText.ToString();
+		}
+
+		private void OpenItemOutputCommandExecuted(object sender, ExecutedRoutedEventArgs e)
+		{
+			if (!(e.Parameter is BuildJob Job)) return;
+
+			DisplayBatchOutputText(Job.OutputText);
+		}
+
+		private void OnDblClickBuildListItem(object sender, MouseButtonEventArgs e)
+		{
+			if (!(sender is FrameworkElement Elem)) return;
+
+			if (!(Elem.DataContext is BuildJob Job)) return;
+
+			// Switch the Startup Project and the build config and platform to match this Job
+			if (Job.Project != null)
+			{
+				IsBusy = true;
+
+				var Project = Job.Project.GetProjectSlow();
+
+				if (Project != null)
+				{
+					// Switch to this project
+					var ProjectHierarchy = Utils.ProjectToHierarchyObject(Project);
+					UnrealVSPackage.Instance.SolutionBuildManager.set_StartupProject(ProjectHierarchy);
+
+					// Switch the build config
+					Utils.SetActiveSolutionConfiguration(Job.Config, Job.Platform);
+				}
+
+				IsBusy = false;
+			}
+		}
+
+		private void OnDblClickBuildingListItem(object sender, MouseButtonEventArgs e)
+		{
+			if (!(sender is FrameworkElement Elem)) return;
+
+			if (!(Elem.DataContext is BuildJob Job)) return;
+
+			DisplayBatchOutputText(Job.OutputText);
+		}
+
+		////////////// Fields ///////////////////////
+
+		private const string BuildJobsPanelPrefix = "Build Jobs";
+		private const string OutputPanelTitlePrefix = "Output";
+
+		private readonly Queue<BuildJob> _BuildQueue = new Queue<BuildJob>();
+		private BuildJob _ActiveBuildJob;
+		private DateTime _BuildJobStartTime;
+		private BuildJobSet _LastSelectedBuildJobSet;
+
+		private readonly ObservableCollection<ProjectListItem> _ProjectCollection = new ObservableCollection<ProjectListItem>();
+		private readonly ObservableCollection<string> _ConfigCollection = new ObservableCollection<string>();
+		private readonly ObservableCollection<string> _PlatformCollection = new ObservableCollection<string>();
+	}
+}
\ No newline at end of file
diff --git a/Engine/Source/Programs/UnrealVS/UnrealVS.2019/BatchBuilderToolWindow.cs b/Engine/Source/Programs/UnrealVS/UnrealVS.2019/BatchBuilderToolWindow.cs
new file mode 100644
--- /dev/null
+++ b/Engine/Source/Programs/UnrealVS/UnrealVS.2019/BatchBuilderToolWindow.cs
@@ -0,0 +1,34 @@
+// Copyright Epic Games, Inc. All Rights Reserved.
+
+using Microsoft.VisualStudio.Shell;
+using System;
+using System.Runtime.InteropServices;
+
+namespace UnrealVS
+{
+	[Guid(GuidList.UnrealVSBatchBuildToolWindowPersistanceString)]
+	internal class BatchBuilderToolWindow : ToolWindowPane
+	{
+		/// <summary>
+		/// Standard constructor for the tool window.
+		/// </summary>
+		public BatchBuilderToolWindow() :
+			base(null)
+		{
+			// Set the window title reading it from the resources.
+			Caption = Resources.BatchBuilderWindowTitle;
+			// Set the image that will appear on the tab of the window frame
+			// when docked with an other window
+			// The resource ID correspond to the one defined in the resx file
+			// while the Index is the offset in the bitmap strip. Each image in
+			// the strip being 16x16.
+			BitmapResourceID = 200;
+			BitmapIndex = 0;
+
+			// This is the user control hosted by the tool window; Note that, even if this class implements IDisposable,
+			// we are not calling Dispose on this object. This is because ToolWindowPane calls Dispose on 
+			// the object returned by the Content property.
+			base.Content = BatchBuilder.ToolControl;
+		}
+	}
+}
diff --git a/Engine/Source/Programs/UnrealVS/UnrealVS.2019/BuildStartupProject.cs b/Engine/Source/Programs/UnrealVS/UnrealVS.2019/BuildStartupProject.cs
new file mode 100644
--- /dev/null
+++ b/Engine/Source/Programs/UnrealVS/UnrealVS.2019/BuildStartupProject.cs
@@ -0,0 +1,59 @@
+// Copyright Epic Games, Inc. All Rights Reserved.
+
+using EnvDTE80;
+using Microsoft.VisualStudio.Shell.Interop;
+using System;
+using System.ComponentModel.Design;
+
+namespace UnrealVS
+{
+	class BuildStartupProject
+	{
+		const int BuildStartupProjectButtonID = 0x1070;
+
+		public BuildStartupProject()
+		{
+			// BuildStartupProjectButton
+			{
+				var CommandID = new CommandID(GuidList.UnrealVSCmdSet, BuildStartupProjectButtonID);
+				var BuildStartupProjectButtonCommand = new MenuCommand(new EventHandler(BuildStartupProjectButtonHandler), CommandID);
+				UnrealVSPackage.Instance.MenuCommandService.AddCommand(BuildStartupProjectButtonCommand);
+			}
+
+		}
+
+
+		/// Called when 'BuildStartupProject' button is clicked
+		void BuildStartupProjectButtonHandler(object Sender, EventArgs Args)
+		{
+			// Grab the current startup project
+			UnrealVSPackage.Instance.SolutionBuildManager.get_StartupProject(out IVsHierarchy ProjectHierarchy);
+			if (ProjectHierarchy != null)
+			{
+				var StartupProject = Utils.HierarchyObjectToProject(ProjectHierarchy);
+
+				if (StartupProject != null)
+				{
+					// Get the active solution configuration
+					var ActiveConfiguration =
+						(SolutionConfiguration2)UnrealVSPackage.Instance.DTE.Solution.SolutionBuild.ActiveConfiguration;
+					var SolutionConfiguration = ActiveConfiguration.Name;
+					var SolutionPlatform = ActiveConfiguration.PlatformName;
+
+					// Combine the active solution configuration and platform into a string that Visual Studio
+					// can use to build the startup project (e.g. "Release|x64")
+					var BuildPlatformAndConfiguration = SolutionConfiguration + "|" + SolutionPlatform;
+
+					// Make sure the Output window is visible
+					UnrealVSPackage.Instance.DTE.ExecuteCommand("View.Output");
+
+					// Kick off the build!
+					UnrealVSPackage.Instance.DTE.Solution.SolutionBuild.BuildProject(
+						BuildPlatformAndConfiguration,
+						StartupProject.UniqueName,
+						WaitForBuildToFinish: false);
+				}
+			}
+		}
+	}
+}
diff --git a/Engine/Source/Programs/UnrealVS/UnrealVS.2019/BuildStatusIcon.xaml b/Engine/Source/Programs/UnrealVS/UnrealVS.2019/BuildStatusIcon.xaml
new file mode 100644
--- /dev/null
+++ b/Engine/Source/Programs/UnrealVS/UnrealVS.2019/BuildStatusIcon.xaml
@@ -0,0 +1,30 @@
+<!--Copyright Epic Games, Inc. All Rights Reserved.-->
+
+<UserControl x:Class="UnrealVS.BuildStatusIcon"
+			 xmlns="http://schemas.microsoft.com/winfx/2006/xaml/presentation"
+			 xmlns:x="http://schemas.microsoft.com/winfx/2006/xaml"
+			 xmlns:mc="http://schemas.openxmlformats.org/markup-compatibility/2006" 
+			 xmlns:d="http://schemas.microsoft.com/expression/blend/2008" 
+			 mc:Ignorable="d" 
+			 d:DesignHeight="200" d:DesignWidth="200"
+			 Foreground="Black">
+	<UserControl.Resources>
+		<LinearGradientBrush x:Key="InnerBrush" EndPoint="0.5,1" StartPoint="0.5,0">
+			<GradientStop Color="{Binding RelativeSource={RelativeSource FindAncestor, AncestorType={x:Type UserControl}}, Path=ColorLo}" Offset="1"/>
+			<GradientStop Color="{Binding RelativeSource={RelativeSource FindAncestor, AncestorType={x:Type UserControl}}, Path=ColorHi}" Offset="0"/>
+		</LinearGradientBrush>
+	</UserControl.Resources>
+	<Grid>
+		<Ellipse Fill="Black" />
+		<Ellipse Fill="{DynamicResource InnerBrush}" RenderTransformOrigin="0.5,0.5">
+			<Ellipse.RenderTransform>
+				<TransformGroup>
+					<ScaleTransform ScaleX="0.85" ScaleY="0.85"/>
+					<SkewTransform/>
+					<RotateTransform/>
+					<TranslateTransform/>
+				</TransformGroup>
+			</Ellipse.RenderTransform>
+		</Ellipse>
+	</Grid>
+</UserControl>
diff --git a/Engine/Source/Programs/UnrealVS/UnrealVS.2019/BuildStatusIcon.xaml.cs b/Engine/Source/Programs/UnrealVS/UnrealVS.2019/BuildStatusIcon.xaml.cs
new file mode 100644
--- /dev/null
+++ b/Engine/Source/Programs/UnrealVS/UnrealVS.2019/BuildStatusIcon.xaml.cs
@@ -0,0 +1,112 @@
+// Copyright Epic Games, Inc. All Rights Reserved.
+
+using System.Windows;
+using System.Windows.Controls;
+using System.Windows.Media;
+
+namespace UnrealVS
+{
+	/// <summary>
+	/// Icon control used on the list of building jobs to show the job status
+	/// </summary>
+	public partial class BuildStatusIcon : UserControl
+	{
+		// Dependency properties
+		public static readonly DependencyProperty BuildStatusProperty;
+		private static readonly DependencyProperty ColorLoProperty;
+		private static readonly DependencyProperty ColorHiProperty;
+
+		public BatchBuilderToolControl.BuildJob.BuildJobStatus BuildStatus
+		{
+			get { return (BatchBuilderToolControl.BuildJob.BuildJobStatus)GetValue(BuildStatusProperty); }
+			set { SetValue(BuildStatusProperty, value); }
+		}
+		private Color ColorLo
+		{
+			get { return (Color)GetValue(ColorLoProperty); }
+			set { SetValue(ColorLoProperty, value); }
+		}
+		private Color ColorHi
+		{
+			get { return (Color)GetValue(ColorHiProperty); }
+			set { SetValue(ColorHiProperty, value); }
+		}
+
+		static BuildStatusIcon()
+		{
+			BuildStatusProperty = DependencyProperty.Register("BuildStatus",
+															  typeof(BatchBuilderToolControl.BuildJob.BuildJobStatus),
+															  typeof(BuildStatusIcon),
+															  new FrameworkPropertyMetadata(
+																  BatchBuilderToolControl.BuildJob.BuildJobStatus.Pending,
+																  OnBuildStatusChanged));
+			ColorLoProperty = DependencyProperty.Register("ColorLo",
+															  typeof(Color),
+															  typeof(BuildStatusIcon),
+															  new FrameworkPropertyMetadata(
+																  GetColorLo(BatchBuilderToolControl.BuildJob.BuildJobStatus.Pending)));
+			ColorHiProperty = DependencyProperty.Register("ColorHi",
+															  typeof(Color),
+															  typeof(BuildStatusIcon),
+															  new FrameworkPropertyMetadata(
+																  GetColorHi(BatchBuilderToolControl.BuildJob.BuildJobStatus.Pending)));
+		}
+
+		private static void OnBuildStatusChanged(DependencyObject Obj, DependencyPropertyChangedEventArgs Args)
+		{
+			var ThisObj = (BuildStatusIcon)Obj;
+			if (ThisObj != null)
+			{
+				ThisObj.ColorLo = GetColorLo((BatchBuilderToolControl.BuildJob.BuildJobStatus)Args.NewValue);
+				ThisObj.ColorHi = GetColorHi((BatchBuilderToolControl.BuildJob.BuildJobStatus)Args.NewValue);
+			}
+		}
+
+		public BuildStatusIcon()
+		{
+			InitializeComponent();
+		}
+
+		private static Color GetColorLo(BatchBuilderToolControl.BuildJob.BuildJobStatus Status)
+		{
+			switch (Status)
+			{
+				case BatchBuilderToolControl.BuildJob.BuildJobStatus.Pending:
+					return Colors.LightGray;
+				case BatchBuilderToolControl.BuildJob.BuildJobStatus.Executing:
+					return Colors.SlateGray;
+				case BatchBuilderToolControl.BuildJob.BuildJobStatus.Failed:
+					return Colors.Red;
+				case BatchBuilderToolControl.BuildJob.BuildJobStatus.FailedToStart:
+					return Colors.Red;
+				case BatchBuilderToolControl.BuildJob.BuildJobStatus.Cancelled:
+					return Colors.DarkOrange;
+				case BatchBuilderToolControl.BuildJob.BuildJobStatus.Complete:
+					return Colors.Green;
+				default:
+					return Colors.Blue;
+			}
+		}
+
+		private static Color GetColorHi(BatchBuilderToolControl.BuildJob.BuildJobStatus Status)
+		{
+			switch (Status)
+			{
+				case BatchBuilderToolControl.BuildJob.BuildJobStatus.Pending:
+					return Colors.DarkGray;
+				case BatchBuilderToolControl.BuildJob.BuildJobStatus.Executing:
+					return Colors.DarkGray;
+				case BatchBuilderToolControl.BuildJob.BuildJobStatus.Failed:
+					return Colors.LightCoral;
+				case BatchBuilderToolControl.BuildJob.BuildJobStatus.FailedToStart:
+					return Colors.LightCoral;
+				case BatchBuilderToolControl.BuildJob.BuildJobStatus.Cancelled:
+					return Colors.LightSalmon;
+				case BatchBuilderToolControl.BuildJob.BuildJobStatus.Complete:
+					return Colors.LightGreen;
+				default:
+					return Colors.LightBlue;
+			}
+		}
+	}
+}
diff --git a/Engine/Source/Programs/UnrealVS/UnrealVS.2019/CommandLineEditor.cs b/Engine/Source/Programs/UnrealVS/UnrealVS.2019/CommandLineEditor.cs
new file mode 100644
--- /dev/null
+++ b/Engine/Source/Programs/UnrealVS/UnrealVS.2019/CommandLineEditor.cs
@@ -0,0 +1,530 @@
+// Copyright Epic Games, Inc. All Rights Reserved.
+
+using EnvDTE;
+using EnvDTE80;
+using Microsoft.VisualStudio;
+using Microsoft.VisualStudio.Shell;
+using Microsoft.VisualStudio.Shell.Interop;
+using System;
+using System.Collections.Generic;
+using System.ComponentModel.Design;
+using System.IO;
+using System.Linq;
+using System.Runtime.InteropServices;
+using System.Text.Json;
+using System.Text.Json.Serialization;
+using System.Timers;
+
+namespace UnrealVS
+{
+	internal class CommandLineEditor : IDisposable
+	{
+		private class LaunchSettingsJson
+		{
+			[JsonPropertyName("profiles")]
+			public Dictionary<string, LaunchSettingsProfile> Profiles { get; set; }
+
+			public class LaunchSettingsProfile
+			{
+				[JsonPropertyName("commandName")]
+				public string CommandName { get; set; }
+				[JsonPropertyName("commandLineArgs")]
+				public string CommandLineArgs { get; set; }
+			}
+		}
+		/** constants */
+
+		private const int ComboID = 0x1030;
+		private const int ComboListID = 0x1040;
+		private const int ComboListCountMax = 32;
+
+		/** methods */
+
+		public CommandLineEditor()
+		{
+		}
+
+
+		public void Initialize()
+		{
+			// Create the handlers for our commands
+			{
+				// CommandLineCombo
+				var ComboCommandID = new CommandID(GuidList.UnrealVSCmdSet, ComboID);
+				ComboCommand = new OleMenuCommand(new EventHandler(ComboHandler), ComboCommandID);
+				UnrealVSPackage.Instance.MenuCommandService.AddCommand(ComboCommand);
+
+				// CommandLineComboList
+				var ComboListCommandID = new CommandID(GuidList.UnrealVSCmdSet, ComboListID);
+				ComboCommandList = new OleMenuCommand(new EventHandler(ComboListHandler), ComboListCommandID);
+				UnrealVSPackage.Instance.MenuCommandService.AddCommand(ComboCommandList);
+			}
+
+			// Register for events that we care about
+			UnrealVSPackage.Instance.OnSolutionOpened += UpdateCommandLineCombo;
+			UnrealVSPackage.Instance.OnSolutionClosed += UpdateCommandLineCombo;
+			UnrealVSPackage.Instance.OnStartupProjectChanged += (p) => UpdateCommandLineCombo();
+			UnrealVSPackage.Instance.OnStartupProjectPropertyChanged += OnStartupProjectPropertyChanged;
+			UnrealVSPackage.Instance.OnStartupProjectConfigChanged += OnStartupProjectConfigChanged;
+
+			UpdateCommandLineCombo();
+		}
+
+		/// <summary>
+		/// Called when a startup project property has changed.  We'll refresh our interface.
+		/// </summary>
+		public void OnStartupProjectPropertyChanged(UInt32 itemid, Int32 propid, UInt32 flags)
+		{
+			// Filter out Helix VS plugin sending thousands of VSHPROPID_StateIconIndex in the solution
+			// This event type is not usefull at all for the commandline editor
+			if (propid == (Int32)__VSHPROPID.VSHPROPID_StateIconIndex)
+			{
+				return;
+			}
+
+			UnrealVSPackage.Instance.SolutionBuildManager.get_StartupProject(out IVsHierarchy ProjectHierarchy);
+			if (ProjectHierarchy != null)
+			{
+				// @todo: filter this so that we only respond to changes in the command line property
+				// Setup a timer to prevent any more performance problem from spamming events
+				if (UpdateCommandLineComboTimer == null)
+				{
+					UpdateCommandLineComboTimer = new System.Timers.Timer(1000)
+					{
+						AutoReset = false,
+					};
+					UpdateCommandLineComboTimer.Elapsed += OnUpdateCommandLineCombo;
+				}
+				// Restart timer to raise the event only after 1s of no notification
+				UpdateCommandLineComboTimer.Stop();
+				UpdateCommandLineComboTimer.Start();
+			}
+		}
+
+		/// <summary>
+		/// Timer callback to UpdateCommandLineCombo
+		/// </summary>
+		private void OnUpdateCommandLineCombo(Object _Source, ElapsedEventArgs _Args)
+		{
+			ThreadHelper.Generic.BeginInvoke(UpdateCommandLineCombo);
+		}
+
+		/// <summary>
+		/// Reads options out of the solution file.
+		/// </summary>
+		public void LoadOptions(Stream Stream)
+		{
+			using (BinaryReader Reader = new BinaryReader(Stream))
+			{
+				List<string> CommandLines = new List<string>();
+				int Count = Reader.ReadInt32();
+				for (int CommandLineIdx = 0; CommandLineIdx < Count; CommandLineIdx++)
+				{
+					CommandLines.Add(Reader.ReadString());
+				}
+				ComboList.Clear();
+				ComboList.AddRange(CommandLines.ToArray());
+			}
+		}
+
+		/// <summary>
+		/// Writes options to the solution file.
+		/// </summary>
+		public void SaveOptions(Stream Stream)
+		{
+			using (BinaryWriter Writer = new BinaryWriter(Stream))
+			{
+				string[] CommandLines = ComboList.ToArray();
+				Writer.Write(CommandLines.Length);
+				for (int CommandLineIdx = 0; CommandLineIdx < CommandLines.Length; CommandLineIdx++)
+				{
+					Writer.Write(CommandLines[CommandLineIdx]);
+				}
+			}
+		}
+
+		/// <summary>
+		/// Called when the startup project active config has changed.
+		/// </summary>
+		private void OnStartupProjectConfigChanged(Project Project)
+		{
+			UpdateCommandLineCombo();
+		}
+
+		/// <summary>
+		/// Updates the command-line combo box after the project loaded state has changed
+		/// </summary>
+		private void UpdateCommandLineCombo()
+		{
+			// Enable or disable our command-line selector
+			DesiredCommandLine = null;  // clear this state var used by the combo box handler
+			UnrealVSPackage.Instance.SolutionBuildManager.get_StartupProject(out IVsHierarchy ProjectHierarchy);
+			if (ProjectHierarchy != null)
+			{
+				ComboCommand.Enabled = true;
+				ComboCommandList.Enabled = true;
+			}
+			else
+			{
+				ComboCommand.Enabled = false;
+				ComboCommandList.Enabled = false;
+			}
+
+			string CommandLine = MakeCommandLineComboText();
+			CommitCommandLineToMRU(CommandLine);
+		}
+
+
+		/// <summary>
+		/// Returns a string to display in the command-line combo box, based on project state
+		/// </summary>
+		/// <returns>String to display</returns>
+		private static string MakeCommandLineComboText()
+		{
+			string Text = "";
+
+			if (UnrealVSPackage.Instance.SolutionBuildManager.get_StartupProject(out IVsHierarchy ProjectHierarchy) == VSConstants.S_OK && ProjectHierarchy != null)
+			{
+				Project SelectedStartupProject = Utils.HierarchyObjectToProject(ProjectHierarchy);
+				if (SelectedStartupProject != null)
+				{
+					Configuration SelectedConfiguration = SelectedStartupProject.ConfigurationManager.ActiveConfiguration;
+					if (SelectedConfiguration != null)
+					{
+						if (ProjectHierarchy is IVsBuildPropertyStorage PropertyStorage)
+						{
+							// Query the property store for the debugger arguments
+							string ConfigurationName = String.Format("{0}|{1}", SelectedConfiguration.ConfigurationName, SelectedConfiguration.PlatformName);
+							if (PropertyStorage.GetPropertyValue("LocalDebuggerCommandArguments", ConfigurationName, (uint)_PersistStorageType.PST_USER_FILE, out Text) != VSConstants.S_OK)
+							{
+								if (PropertyStorage.GetPropertyValue("StartArguments", ConfigurationName, (uint)_PersistStorageType.PST_USER_FILE, out Text) != VSConstants.S_OK)
+								{
+									Text = "";
+								}
+							}
+							else
+							{
+								// net core projects move debugger arguments into a debug profile setup via launchSetings.json
+								if (PropertyStorage.GetPropertyValue("ActiveDebugProfile", ConfigurationName, (uint)_PersistStorageType.PST_USER_FILE, out string activeDebugProfile) != VSConstants.S_OK)
+								{
+									activeDebugProfile = null;
+								}
+
+								string launchSettingsPath = Path.Combine(Path.GetDirectoryName(SelectedStartupProject.FileName), "Properties", "launchSettings.json");
+								if (File.Exists(launchSettingsPath))
+								{
+									string jsonDoc = File.ReadAllText(launchSettingsPath);
+									LaunchSettingsJson settings = JsonSerializer.Deserialize<LaunchSettingsJson>(jsonDoc, new JsonSerializerOptions
+									{
+										PropertyNameCaseInsensitive = true
+									});
+
+									if (!string.IsNullOrEmpty(activeDebugProfile))
+									{
+										Text = settings.Profiles[activeDebugProfile].CommandLineArgs ?? "";
+									}
+									else
+									{
+										// if no active debug profile is set then VS uses the first one
+										Text = settings.Profiles.FirstOrDefault().Value?.CommandLineArgs ?? "";
+									}
+								}
+							}
+							// for "Game" projects automatically remove the game project filename from the start of the command line
+							var ActiveConfiguration = (SolutionConfiguration2)UnrealVSPackage.Instance.DTE.Solution.SolutionBuild.ActiveConfiguration;
+							if (UnrealVSPackage.Instance.IsUESolutionLoaded && Utils.IsGameProject(SelectedStartupProject) && Utils.HasUProjectCommandLineArg(ActiveConfiguration.Name))
+							{
+								string AutoPrefix = Utils.GetAutoUProjectCommandLinePrefix(SelectedStartupProject);
+								if (!string.IsNullOrEmpty(AutoPrefix))
+								{
+									if (Text.Trim().StartsWith(AutoPrefix, StringComparison.OrdinalIgnoreCase))
+									{
+										Text = Text.Trim().Substring(AutoPrefix.Length).Trim();
+									}
+								}
+							}
+						}
+					}
+				}
+			}
+
+			return Text;
+		}
+
+
+		/// Called by combo control to query the text to display or to apply newly-entered text
+		private void ComboHandler(object Sender, EventArgs Args)
+		{
+			try
+			{
+				var OleArgs = (OleMenuCmdEventArgs)Args;
+
+				if (OleArgs.InValue is string InString)
+				{
+					// New text set on the combo - set the command line property
+					DesiredCommandLine = null;
+					CommitCommandLineText(InString);
+				}
+				else if (OleArgs.OutValue != IntPtr.Zero)
+				{
+					string EditingString = null;
+					if (OleArgs.InValue != null)
+					{
+						if (OleArgs.InValue is object[] InArray && 0 < InArray.Length)
+						{
+							EditingString = InArray.Last() as string;
+						}
+					}
+
+					string TextToDisplay = string.Empty;
+					if (EditingString != null)
+					{
+						// The control wants to put EditingString in the box
+						TextToDisplay = DesiredCommandLine = EditingString;
+					}
+					else
+					{
+						// This is always hit at the end of interaction with the combo
+						if (DesiredCommandLine != null)
+						{
+							TextToDisplay = DesiredCommandLine;
+							DesiredCommandLine = null;
+							CommitCommandLineText(TextToDisplay);
+						}
+						else
+						{
+							TextToDisplay = MakeCommandLineComboText();
+						}
+					}
+
+					Marshal.GetNativeVariantForObject(TextToDisplay, OleArgs.OutValue);
+				}
+			}
+			catch (Exception ex)
+			{
+				Exception AppEx = new ApplicationException("CommandLineEditor threw an exception in ComboHandler()", ex);
+				Logging.WriteLine(AppEx.ToString());
+				throw AppEx;
+			}
+		}
+
+		private void CommitCommandLineText(string CommandLine)
+		{
+			if (UnrealVSPackage.Instance.SolutionBuildManager.get_StartupProject(out IVsHierarchy ProjectHierarchy) == VSConstants.S_OK && ProjectHierarchy != null)
+			{
+				Project SelectedStartupProject = Utils.HierarchyObjectToProject(ProjectHierarchy);
+				if (SelectedStartupProject != null)
+				{
+					Configuration SelectedConfiguration = SelectedStartupProject.ConfigurationManager.ActiveConfiguration;
+					if (SelectedConfiguration != null)
+					{
+						if (ProjectHierarchy is IVsBuildPropertyStorage PropertyStorage)
+						{
+							string FullCommandLine = CommandLine;
+
+							// for "Game" projects automatically remove the game project filename from the start of the command line
+							var ActiveConfiguration = (SolutionConfiguration2)UnrealVSPackage.Instance.DTE.Solution.SolutionBuild.ActiveConfiguration;
+							if (UnrealVSPackage.Instance.IsUESolutionLoaded && Utils.IsGameProject(SelectedStartupProject) && Utils.HasUProjectCommandLineArg(ActiveConfiguration.Name))
+							{
+								string AutoPrefix = Utils.GetAutoUProjectCommandLinePrefix(SelectedStartupProject);
+								if (FullCommandLine.IndexOf(Utils.UProjectExtension, StringComparison.OrdinalIgnoreCase) < 0 &&
+									string.Compare(FullCommandLine, SelectedStartupProject.Name, StringComparison.OrdinalIgnoreCase) != 0 &&
+									!FullCommandLine.StartsWith(SelectedStartupProject.Name + " ", StringComparison.OrdinalIgnoreCase))
+								{
+									// Committed command line does not specify a .uproject
+									FullCommandLine = AutoPrefix + " " + FullCommandLine;
+								}
+							}
+
+							// Get the project platform name.
+							string ProjectPlatformName = SelectedConfiguration.PlatformName;
+							if (ProjectPlatformName == "Any CPU")
+							{
+								ProjectPlatformName = "AnyCPU";
+							}
+
+							// Get the project kind. C++ projects store the debugger arguments differently to other project types.
+							string ProjectKind = SelectedStartupProject.Kind;
+
+							// Update the property
+							string ProjectConfigurationName = String.Format("{0}|{1}", SelectedConfiguration.ConfigurationName, ProjectPlatformName);
+							if (String.Equals(ProjectKind, "{8BC9CEB8-8B4A-11D0-8D11-00A0C91BC942}", StringComparison.InvariantCultureIgnoreCase))
+							{
+								List<string> ExtraFields = Utils.GetExtraDebuggerCommandArguments(ProjectPlatformName, SelectedStartupProject);
+
+								if (FullCommandLine.Length == 0)
+								{
+									PropertyStorage.RemoveProperty("LocalDebuggerCommandArguments", ProjectConfigurationName, (uint)_PersistStorageType.PST_USER_FILE);
+									PropertyStorage.RemoveProperty("RemoteDebuggerCommandArguments", ProjectConfigurationName, (uint)_PersistStorageType.PST_USER_FILE);
+									foreach (string ExtraField in ExtraFields)
+									{
+										PropertyStorage.RemoveProperty(ExtraField, ProjectConfigurationName, (uint)_PersistStorageType.PST_USER_FILE);
+									}
+								}
+								else
+								{
+									PropertyStorage.SetPropertyValue("LocalDebuggerCommandArguments", ProjectConfigurationName, (uint)_PersistStorageType.PST_USER_FILE, FullCommandLine);
+									PropertyStorage.SetPropertyValue("RemoteDebuggerCommandArguments", ProjectConfigurationName, (uint)_PersistStorageType.PST_USER_FILE, FullCommandLine);
+									foreach (string ExtraField in ExtraFields)
+									{
+										PropertyStorage.SetPropertyValue(ExtraField, ProjectConfigurationName, (uint)_PersistStorageType.PST_USER_FILE, FullCommandLine);
+									}
+								}
+							}
+							else
+							{
+								// For some reason, have to update C# projects this way, otherwise the properties page doesn't update. Conversely, SelectedConfiguration.Properties is always null for C++ projects in VS2017.
+								bool netCoreProject = false;
+								if (SelectedConfiguration.Properties != null)
+								{
+									foreach (Property Property in SelectedConfiguration.Properties)
+									{
+										if (Property.Name.Equals("StartArguments", StringComparison.InvariantCultureIgnoreCase))
+										{
+											try
+											{
+												Property.Value = FullCommandLine;
+											}
+											catch (NotImplementedException)
+											{
+												// net core projects report containing start arguments command but does not support setting it
+												netCoreProject = true;
+											}
+											break;
+										}
+									}
+								}
+
+								if (netCoreProject)
+								{
+									// net core project use launchSettings.json to control debug arguments instead
+
+									if (PropertyStorage.GetPropertyValue("ActiveDebugProfile", ProjectConfigurationName, (uint)_PersistStorageType.PST_USER_FILE, out string activeDebugProfile) != VSConstants.S_OK)
+									{
+										activeDebugProfile = null;
+									}
+
+									string launchSettingsPath = Path.Combine(Path.GetDirectoryName(SelectedStartupProject.FileName), "Properties", "launchSettings.json");
+									LaunchSettingsJson settings = null;
+									if (File.Exists(launchSettingsPath))
+									{
+										try
+										{
+											string jsonDoc = File.ReadAllText(launchSettingsPath);
+											settings = JsonSerializer.Deserialize<LaunchSettingsJson>(jsonDoc, new JsonSerializerOptions
+											{
+												PropertyNameCaseInsensitive = true
+											});
+										}
+										catch (Exception Ex)
+										{
+											Logging.WriteLine($"Error loading {launchSettingsPath}");
+											Logging.WriteLine(Ex.ToString());
+										}
+									}
+
+									// Create empty Launch Settings if launchSettings.json does not exist or failed to load
+									if (settings == null || settings.Profiles == null || settings.Profiles.Count == 0)
+									{
+										settings = new LaunchSettingsJson
+										{
+											Profiles = new Dictionary<string, LaunchSettingsJson.LaunchSettingsProfile>()
+										};
+										settings.Profiles.Add(SelectedStartupProject.Name, new LaunchSettingsJson.LaunchSettingsProfile
+										{
+											CommandName = "Project",
+											CommandLineArgs = string.Empty,
+										});
+									}
+
+									if (!string.IsNullOrEmpty(activeDebugProfile) && settings.Profiles.ContainsKey(activeDebugProfile))
+									{
+										settings.Profiles[activeDebugProfile].CommandLineArgs = FullCommandLine;
+									}
+									else
+									{
+										// if no active debug profile is set then VS uses the first one
+										settings.Profiles.First().Value.CommandLineArgs = FullCommandLine;
+									}
+
+									string json = JsonSerializer.Serialize(settings, new JsonSerializerOptions
+									{
+										WriteIndented = true,
+									});
+									File.WriteAllText(launchSettingsPath, json);
+								}
+							}
+						}
+					}
+				}
+			}
+
+			CommitCommandLineToMRU(CommandLine);
+		}
+
+		private void CommitCommandLineToMRU(string CommandLine)
+		{
+			if (0 < CommandLine.Length)
+			{
+				// Maintain the MRU history
+				// Adds the entered command line to the top of the list
+				// Moves it to the top if it's already in the list
+				// Trims the list to the max length if it's too big
+				ComboList.RemoveAll(s => 0 == string.Compare(s, CommandLine));
+				ComboList.Insert(0, CommandLine);
+				if (ComboList.Count > ComboListCountMax)
+				{
+					ComboList.RemoveAt(ComboList.Count - 1);
+				}
+			}
+		}
+
+		/// Called by combo control to populate the drop-down list
+		private void ComboListHandler(object Sender, EventArgs Args)
+		{
+			var OleArgs = (OleMenuCmdEventArgs)Args;
+
+			Marshal.GetNativeVariantForObject(ComboList.ToArray(), OleArgs.OutValue);
+		}
+
+		/// List of MRU strings to show in the combobox drop-down list
+		private readonly List<string> ComboList = new List<string>();
+
+		/// Combo control for command-line editor
+		private OleMenuCommand ComboCommand;
+		private OleMenuCommand ComboCommandList;
+
+		/// Used to store the user edited commandline mid-edit, in the combo handler
+		private string DesiredCommandLine;
+
+		/// used to shield against too many property changes in the future
+		private System.Timers.Timer UpdateCommandLineComboTimer = null;
+
+		// Track IDisposable state
+		private bool IsDisposed;
+
+		protected virtual void Dispose(bool disposing)
+		{
+			if (!IsDisposed)
+			{
+				if (disposing)
+				{
+					// Dispose managed state (managed objects)
+					UpdateCommandLineComboTimer?.Dispose();
+					UpdateCommandLineComboTimer = null;
+				}
+
+				// Free unmanaged resources (unmanaged objects) and override finalizer
+				// Set large fields to null
+				IsDisposed = true;
+			}
+		}
+
+		public void Dispose()
+		{
+			// Do not change this code. Put cleanup code in 'Dispose(bool disposing)' method
+			Dispose(disposing: true);
+			GC.SuppressFinalize(this);
+		}
+	}
+}
diff --git a/Engine/Source/Programs/UnrealVS/UnrealVS.2019/CompileSingleFile.cs b/Engine/Source/Programs/UnrealVS/UnrealVS.2019/CompileSingleFile.cs
new file mode 100644
--- /dev/null
+++ b/Engine/Source/Programs/UnrealVS/UnrealVS.2019/CompileSingleFile.cs
@@ -0,0 +1,191 @@
+// Copyright Epic Games, Inc. All Rights Reserved.
+
+using EnvDTE;
+using Microsoft.VisualStudio.Shell.Interop;
+using System;
+using System.Collections.Generic;
+using System.ComponentModel.Design;
+using System.Diagnostics;
+using System.IO;
+using System.Windows.Forms;
+
+namespace UnrealVS
+{
+	class CompileSingleFile : IDisposable
+	{
+		const int CompileSingleFileButtonID = 0x1075;
+		static readonly List<string> ValidExtensions = new List<string> { ".c", ".cc", ".cpp", ".cxx" };
+
+		System.Diagnostics.Process ChildProcess;
+
+		public CompileSingleFile()
+		{
+			CommandID CommandID = new CommandID(GuidList.UnrealVSCmdSet, CompileSingleFileButtonID);
+			MenuCommand CompileSingleFileButtonCommand = new MenuCommand(new EventHandler(CompileSingleFileButtonHandler), CommandID);
+			UnrealVSPackage.Instance.MenuCommandService.AddCommand(CompileSingleFileButtonCommand);
+		}
+
+		public void Dispose()
+		{
+			KillChildProcess();
+		}
+
+		void CompileSingleFileButtonHandler(object Sender, EventArgs Args)
+		{
+			if (!TryCompileSingleFile())
+			{
+				DTE DTE = UnrealVSPackage.Instance.DTE;
+				DTE.ExecuteCommand("Build.Compile");
+			}
+		}
+
+		void KillChildProcess()
+		{
+			if (ChildProcess != null)
+			{
+				if (!ChildProcess.HasExited)
+				{
+					ChildProcess.Kill();
+					ChildProcess.WaitForExit();
+				}
+				ChildProcess.Dispose();
+				ChildProcess = null;
+			}
+		}
+
+		bool TryCompileSingleFile()
+		{
+			DTE DTE = UnrealVSPackage.Instance.DTE;
+
+			// Activate the output window
+			Window Window = DTE.Windows.Item(EnvDTE.Constants.vsWindowKindOutput);
+			Window.Activate();
+
+			// Find or create the 'Build' window
+			IVsOutputWindowPane BuildOutputPane = UnrealVSPackage.Instance.GetOutputPane();
+			if (BuildOutputPane == null)
+			{
+				Logging.WriteLine("CompileSingleFile: Build Output Pane not found");
+				return false;
+			}
+
+			// If there's already a build in progress, offer to cancel it
+			if (ChildProcess != null && !ChildProcess.HasExited)
+			{
+				if (MessageBox.Show("Cancel current compile?", "Compile in progress", MessageBoxButtons.YesNo) == DialogResult.Yes)
+				{
+					KillChildProcess();
+					BuildOutputPane.OutputString($"1>  Build cancelled.{Environment.NewLine}");
+				}
+				return true;
+			}
+
+			// Check we've got a file open
+			if (DTE.ActiveDocument == null)
+			{
+				Logging.WriteLine("CompileSingleFile: ActiveDocument not found");
+				return false;
+			}
+
+			// Grab the current startup project
+			UnrealVSPackage.Instance.SolutionBuildManager.get_StartupProject(out IVsHierarchy ProjectHierarchy);
+			if (ProjectHierarchy == null)
+			{
+				Logging.WriteLine("CompileSingleFile: ProjectHierarchy not found");
+				return false;
+			}
+			Project StartupProject = Utils.HierarchyObjectToProject(ProjectHierarchy);
+			if (StartupProject == null)
+			{
+				Logging.WriteLine("CompileSingleFile: StartupProject not found");
+				return false;
+			}
+			if (!(StartupProject.Object is Microsoft.VisualStudio.VCProjectEngine.VCProject VCStartupProject))
+			{
+				Logging.WriteLine("CompileSingleFile: VCStartupProject not found");
+				return false;
+			}
+
+			// Get the active configuration for the startup project
+			Configuration ActiveConfiguration = StartupProject.ConfigurationManager.ActiveConfiguration;
+			string ActiveConfigurationName = $"{ActiveConfiguration.ConfigurationName}|{ActiveConfiguration.PlatformName}";
+			Microsoft.VisualStudio.VCProjectEngine.VCConfiguration ActiveVCConfiguration = (VCStartupProject.Configurations as Microsoft.VisualStudio.VCProjectEngine.IVCCollection).Item(ActiveConfigurationName) as Microsoft.VisualStudio.VCProjectEngine.VCConfiguration;
+			if (ActiveVCConfiguration == null)
+			{
+				Logging.WriteLine("CompileSingleFile: VCStartupProject ActiveConfiguration not found");
+				return false;
+			}
+
+			// Get the NMake settings for this configuration
+			Microsoft.VisualStudio.VCProjectEngine.VCNMakeTool ActiveNMakeTool = (ActiveVCConfiguration.Tools as Microsoft.VisualStudio.VCProjectEngine.IVCCollection).Item("VCNMakeTool") as Microsoft.VisualStudio.VCProjectEngine.VCNMakeTool;
+			if (ActiveNMakeTool == null)
+			{
+				MessageBox.Show($"No NMakeTool set for Project {VCStartupProject.Name} set for single-file compile.", "NMakeTool not set", MessageBoxButtons.OK);
+				return false;
+			}
+
+			// Save all the open documents
+			DTE.ExecuteCommand("File.SaveAll");
+
+			// Check if the requested file is valid
+			string FileToCompile = DTE.ActiveDocument.FullName;
+			string FileToCompileExt = Path.GetExtension(FileToCompile);
+			if (!ValidExtensions.Contains(FileToCompileExt.ToLowerInvariant()))
+			{
+				MessageBox.Show($"Invalid file extension {FileToCompileExt} for single-file compile.", "Invalid Extension", MessageBoxButtons.OK);
+				return true;
+			}
+
+			// If there's already a build in progress, don't let another one start
+			if (DTE.Solution.SolutionBuild.BuildState == vsBuildState.vsBuildStateInProgress)
+			{
+				if (MessageBox.Show("Cancel current compile?", "Compile in progress", MessageBoxButtons.YesNo) == DialogResult.Yes)
+				{
+					DTE.ExecuteCommand("Build.Cancel");
+				}
+				return true;
+			}
+
+			// Make sure any existing build is stopped
+			KillChildProcess();
+
+			// Set up the output pane
+			BuildOutputPane.Activate();
+			BuildOutputPane.Clear();
+			BuildOutputPane.OutputString($"1>------ Build started: Project: {StartupProject.Name}, Configuration: {ActiveConfiguration.ConfigurationName} {ActiveConfiguration.PlatformName} ------{Environment.NewLine}");
+			BuildOutputPane.OutputString($"1>  Compiling {FileToCompile}{Environment.NewLine}");
+
+			// Set up event handlers 
+			DTE.Events.BuildEvents.OnBuildBegin += BuildEvents_OnBuildBegin;
+
+			// Create a delegate for handling output messages
+			void OutputHandler(object Sender, DataReceivedEventArgs Args) { if (Args.Data != null) BuildOutputPane.OutputString($"1>  {Args.Data}{Environment.NewLine}"); }
+
+			// Get the build command line and escape any environment variables that we use
+			string BuildCommandLine = ActiveNMakeTool.BuildCommandLine;
+			BuildCommandLine = BuildCommandLine.Replace("$(SolutionDir)", Path.GetDirectoryName(UnrealVSPackage.Instance.SolutionFilepath).TrimEnd(Path.DirectorySeparatorChar) + Path.DirectorySeparatorChar);
+			BuildCommandLine = BuildCommandLine.Replace("$(ProjectName)", VCStartupProject.Name);
+
+			// Spawn the new process
+			ChildProcess = new System.Diagnostics.Process();
+			ChildProcess.StartInfo.FileName = Path.Combine(Environment.SystemDirectory, "cmd.exe");
+			ChildProcess.StartInfo.Arguments = $"/C \"{BuildCommandLine} -singlefile=\"{FileToCompile}\"\"";
+			ChildProcess.StartInfo.WorkingDirectory = Path.GetDirectoryName(StartupProject.FullName);
+			ChildProcess.StartInfo.UseShellExecute = false;
+			ChildProcess.StartInfo.RedirectStandardOutput = true;
+			ChildProcess.StartInfo.RedirectStandardError = true;
+			ChildProcess.StartInfo.CreateNoWindow = true;
+			ChildProcess.OutputDataReceived += OutputHandler;
+			ChildProcess.ErrorDataReceived += OutputHandler;
+			ChildProcess.Start();
+			ChildProcess.BeginOutputReadLine();
+			ChildProcess.BeginErrorReadLine();
+			return true;
+		}
+
+		private void BuildEvents_OnBuildBegin(vsBuildScope Scope, vsBuildAction Action)
+		{
+			KillChildProcess();
+		}
+	}
+}
diff --git a/Engine/Source/Programs/UnrealVS/UnrealVS.2019/GenerateProjectFiles.cs b/Engine/Source/Programs/UnrealVS/UnrealVS.2019/GenerateProjectFiles.cs
new file mode 100644
--- /dev/null
+++ b/Engine/Source/Programs/UnrealVS/UnrealVS.2019/GenerateProjectFiles.cs
@@ -0,0 +1,172 @@
+// Copyright Epic Games, Inc. All Rights Reserved.
+
+using System;
+using System.ComponentModel.Design;
+using System.Diagnostics;
+using System.IO;
+
+
+namespace UnrealVS
+{
+	class GenerateProjectFiles
+	{
+		const int GenerateProjectFilesButtonID = 0x1100;
+
+
+		public GenerateProjectFiles()
+		{
+			// GenerateProjectFilesButton
+			{
+				var CommandID = new CommandID(GuidList.UnrealVSCmdSet, GenerateProjectFilesButtonID);
+				GenerateProjectFilesButtonCommand = new MenuCommand(new EventHandler(GenerateProjectFilesButtonHandler), CommandID);
+				UnrealVSPackage.Instance.MenuCommandService.AddCommand(GenerateProjectFilesButtonCommand);
+			}
+
+
+			// Register to find out when the solution is closed or opened
+			UnrealVSPackage.Instance.OnSolutionOpened += OnSolutionOpened;
+			UnrealVSPackage.Instance.OnSolutionClosing += OnSolutionClosing;
+
+			UpdateGenerateProjectFilesButtonState();
+		}
+
+
+
+		/// Called when 'GenerateProjectFiles' button is clicked
+		void GenerateProjectFilesButtonHandler(object Sender, EventArgs Args)
+		{
+			// Don't allow a regen if we're already generating projects
+			if (GenerateProjectFilesProcess == null)
+			{
+				// Figure out which project to rebuild
+				string BatchFileName = GetBatchFileName();
+
+				if (!String.IsNullOrEmpty(BatchFileName))
+				{
+					// Make sure the Output window is visible
+					UnrealVSPackage.Instance.DTE.ExecuteCommand("View.Output");
+
+					// Activate the output pane
+					var Pane = UnrealVSPackage.Instance.GetOutputPane();
+					if (Pane != null)
+					{
+						// Clear and activate the output pane.
+						Pane.Clear();
+
+						// @todo: Activating doesn't seem to really bring the pane to front like we would expect it to.
+						Pane.Activate();
+					}
+
+					GenerateProjectFilesProcess = UnrealVSPackage.LaunchProgram(
+						"cmd.exe", "/C " + GetSafeFilePath(BatchFileName),
+						OnGenerateProjectFilesProcessExit,
+						OnOutputFromGenerateProjectFilesProcess);
+				}
+
+				UpdateGenerateProjectFilesButtonState();
+			}
+		}
+
+
+		/// Called when the GenerateProjectFiles process terminates
+		void OnGenerateProjectFilesProcessExit(object Sender, EventArgs Args)
+		{
+			GenerateProjectFilesProcess = null;
+
+			UpdateGenerateProjectFilesButtonState();
+		}
+
+
+		/// <summary>
+		/// Updates the enabled/disabled state of the 'Generate Project Files' button
+		/// </summary>
+		void UpdateGenerateProjectFilesButtonState(bool ShouldForceDisable = false)
+		{
+			var CanGenerateProjects = false;
+
+			if (!ShouldForceDisable)
+			{
+				// Can't launch when a process is already running
+				if (GenerateProjectFilesProcess == null)
+				{
+					string BatchFileName = GetBatchFileName();
+					if (BatchFileName != null)
+					{
+						CanGenerateProjects = true;
+					}
+					else
+					{
+						// No solution loaded, or solution doesn't look like Unreal
+					}
+				}
+			}
+
+			GenerateProjectFilesButtonCommand.Enabled = CanGenerateProjects;
+		}
+
+
+		/// <summary>
+		/// Called when a solution is opened
+		/// </summary>
+		public void OnSolutionOpened()
+		{
+			UpdateGenerateProjectFilesButtonState();
+		}
+
+
+		/// <summary>
+		/// Called when the solution is closed
+		/// </summary>
+		public void OnSolutionClosing()
+		{
+			UpdateGenerateProjectFilesButtonState(ShouldForceDisable: true);
+		}
+
+
+		/// <summary>
+		/// Called when data is received from the GenerateProjectFiles process
+		/// </summary>
+		public void OnOutputFromGenerateProjectFilesProcess(object _Sender, DataReceivedEventArgs Args)
+		{
+			var Pane = UnrealVSPackage.Instance.GetOutputPane();
+			if (Pane != null)
+			{
+				Pane.OutputString(Args.Data + "\n");
+			}
+		}
+
+		private string GetBatchFileName()
+		{
+			// Check to see if we have the Unreal loaded
+			if (UnrealVSPackage.Instance.IsUESolutionLoaded)
+			{
+				// We expect "GenerateProjectFiles.bat" to live in the same directory as the solution
+				return Path.Combine(Path.GetDirectoryName(UnrealVSPackage.Instance.SolutionFilepath), "GenerateProjectFiles.bat");
+			}
+			return null;
+		}
+
+		/// <summary>
+		/// Wraps a file path in quotes if it contains a space character
+		/// </summary>
+		/// <param name="InPath"></param>
+		/// <returns></returns>
+		private string GetSafeFilePath(string InPath)
+		{
+			string WorkingPath = InPath;
+
+			if (WorkingPath.Contains(" ") && !WorkingPath.Contains("\""))
+			{
+				WorkingPath = "\"" + WorkingPath + "\"";
+			}
+			return WorkingPath;
+		}
+
+		/// Active process for the "generate project files" command
+		System.Diagnostics.Process GenerateProjectFilesProcess;
+
+		/// Command for 'Generate Project Files' button
+		readonly MenuCommand GenerateProjectFilesButtonCommand;
+	}
+
+}
diff --git a/Engine/Source/Programs/UnrealVS/UnrealVS.2019/GlobalSuppressions.cs b/Engine/Source/Programs/UnrealVS/UnrealVS.2019/GlobalSuppressions.cs
new file mode 100644
--- /dev/null
+++ b/Engine/Source/Programs/UnrealVS/UnrealVS.2019/GlobalSuppressions.cs
@@ -0,0 +1,13 @@
+// Copyright Epic Games, Inc. All Rights Reserved.
+
+// This file is used by Code Analysis to maintain SuppressMessage
+// attributes that are applied to this project. Project-level
+// suppressions either have no target or are given a specific target
+// and scoped to a namespace, type, member, etc.
+//
+// To add a suppression to this file, right-click the message in the
+// Error List, point to "Suppress Message(s)", and click "In Project
+// Suppression File". You do not need to add suppressions to this
+// file manually.
+
+[assembly: System.Diagnostics.CodeAnalysis.SuppressMessage("Microsoft.Design", "CA1017:MarkAssembliesWithComVisible")]
diff --git a/Engine/Source/Programs/UnrealVS/UnrealVS.2019/Guids.cs b/Engine/Source/Programs/UnrealVS/UnrealVS.2019/Guids.cs
new file mode 100644
--- /dev/null
+++ b/Engine/Source/Programs/UnrealVS/UnrealVS.2019/Guids.cs
@@ -0,0 +1,33 @@
+// Copyright Epic Games, Inc. All Rights Reserved.
+
+// Guids.cs
+// MUST match guids.h
+using System;
+
+namespace UnrealVS
+{
+	static class GuidList
+	{
+		public const string UnrealVSPackageString = "ddbf523f-7eb6-4887-bd51-85a714ff87eb";
+		public const string UnrealVSCmdSetString = "f10be850-eeec-4cef-ac64-1bc2cbe0d447";
+		public const string UnrealVSPackageManString = "683EF1CF-9C97-46D7-BCFA-F12CA6FB77D6";
+		public const string UnrealVSBatchBuildToolWindowPersistanceString = "9DE31877-5741-4B8B-9F3B-633555D7842B";
+		public const string UnrealVSOptionsString = "FAC770BD-A983-4F40-9364-868D645AF16D";
+
+		// spare generated GUIDs
+		//[Guid("DFF26599-0F46-4A6B-8204-45CC89BAB09D")]
+		//[Guid("E7E7B844-44CF-4DCE-A650-3F97A63C5743")]
+
+		public static readonly Guid UnrealVSCmdSet = new Guid(UnrealVSCmdSetString);
+
+		public static readonly Guid BatchBuildPaneGuid = new Guid("BBAC0710-4A2E-4BD3-9A39-4683FC6CFE86");
+
+		// Known GUIDs - not defined in the SDK? (found at http://msdn.microsoft.com/en-us/library/hb23x61k(v=VS.80).aspx)
+		public const string SolutionFolderGuidString = "{66A26720-8FB5-11D2-AA7E-00C04F688DDE}";
+		public const string VBProjectKindGuidString = "{F184B08F-C81C-45F6-A57F-5ABD9991F28F}";
+		public const string VCSharpProjectKindGuidString = "{FAE04EC0-301F-11D3-BF4B-00C04F79EFBC}";
+		public const string VCProjectKindGuidString = "{8BC9CEB8-8B4A-11D0-8D11-00A0C91BC942}";
+		public const string VJSharpProjectKindGuidString = "{E6FDF86B-F3D1-11D4-8576-0002A516ECE8}";
+		public const string WebProjectKindGuidString = "{E24C65DC-7377-472b-9ABA-BC803B73C61A}";
+	};
+}
\ No newline at end of file
diff --git a/Engine/Source/Programs/UnrealVS/UnrealVS.2019/Logging.cs b/Engine/Source/Programs/UnrealVS/UnrealVS.2019/Logging.cs
new file mode 100644
--- /dev/null
+++ b/Engine/Source/Programs/UnrealVS/UnrealVS.2019/Logging.cs
@@ -0,0 +1,138 @@
+// Copyright Epic Games, Inc. All Rights Reserved.
+
+using System;
+using System.Diagnostics;
+using System.IO;
+
+namespace UnrealVS
+{
+	static class Logging
+	{
+		private static readonly string LogFolderPathRoot =
+			Path.Combine(Environment.GetFolderPath(Environment.SpecialFolder.LocalApplicationData), "Epic Games");
+
+		private const string LogFileNameBase = "UnrealVS-log";
+		private const string LogFileNameExt = ".txt";
+
+		private static StreamWriter LogFile;
+		private static bool bLoggingReady = false;
+
+		private static string ExtensionName;
+		private static string VersionString;
+
+		private const int MaxFileSuffix = 64;
+
+		public static void Initialize(string InExtensionName, string InVersionString)
+		{
+			ExtensionName = InExtensionName;
+			VersionString = InVersionString;
+
+			Initialize(0);
+		}
+
+		private static void Initialize(int FileSuffix)
+		{
+			try
+			{
+				string LogFolderPath = Path.Combine(LogFolderPathRoot, ExtensionName);
+
+				if (!Directory.Exists(LogFolderPath))
+				{
+					Directory.CreateDirectory(LogFolderPath);
+				}
+
+				string LogFilePath = GetLogFilePath(LogFolderPath, FileSuffix);
+
+				try
+				{
+					if (File.Exists(LogFilePath))
+					{
+						File.Delete(LogFilePath);
+					}
+
+					LogFile = new StreamWriter(LogFilePath);
+					bLoggingReady = true;
+
+					WriteLine(
+						string.Format(
+							"LOG: {0} {1} (started up at {2} - {3})",
+							ExtensionName,
+							VersionString,
+							DateTime.Now.ToLongDateString(),
+							DateTime.Now.ToLongTimeString()));
+
+#if VS11
+					WriteLine("Visual Studio 11 build");
+#elif VS12
+					WriteLine("Visual Studio 12 build");
+#elif VS13
+					WriteLine("Visual Studio 13 build");
+#elif VS14
+					WriteLine("Visual Studio 14 build");
+#elif VS15
+					WriteLine("Visual Studio 15 build");
+#elif VS16
+					WriteLine("Visual Studio 16 build");
+#else
+					WriteLine("UNKNOWN Visual Studio build");
+#endif
+				}
+				catch (IOException)
+				{
+					if (MaxFileSuffix == FileSuffix) throw;
+
+					Initialize(FileSuffix + 1);
+				}
+			}
+			catch (Exception ex)
+			{
+				if (ex is ApplicationException) throw;
+
+				bLoggingReady = false;
+				throw new ApplicationException("Failed to init logging in UnrealVS", ex);
+			}
+		}
+
+		private static string GetLogFilePath(string LogFolderPath, int FileSuffix)
+		{
+			string Suffix = string.Empty;
+			if (FileSuffix > 0)
+			{
+				Suffix = string.Format("({0})", FileSuffix);
+			}
+
+			return LogFolderPath +
+					Path.DirectorySeparatorChar +
+					LogFileNameBase +
+					Suffix +
+					LogFileNameExt;
+		}
+
+		public static void Close()
+		{
+			if (!bLoggingReady) return;
+
+			WriteLine(
+				string.Format(
+					"LOG: {0} {1} (closed at {2} - {3})",
+					ExtensionName,
+					VersionString,
+					DateTime.Now.ToLongDateString(),
+					DateTime.Now.ToLongTimeString()));
+
+			LogFile.Close();
+			LogFile = null;
+			bLoggingReady = false;
+		}
+
+		public static void WriteLine(string Text)
+		{
+			Trace.WriteLine(Text);
+
+			if (!bLoggingReady) return;
+
+			LogFile.WriteLine(Text);
+			LogFile.Flush();
+		}
+	}
+}
diff --git a/Engine/Source/Programs/UnrealVS/UnrealVS.2019/ProfileManager.cs b/Engine/Source/Programs/UnrealVS/UnrealVS.2019/ProfileManager.cs
new file mode 100644
--- /dev/null
+++ b/Engine/Source/Programs/UnrealVS/UnrealVS.2019/ProfileManager.cs
@@ -0,0 +1,41 @@
+// Copyright Epic Games, Inc. All Rights Reserved.
+
+using Microsoft.VisualStudio.Shell;
+using Microsoft.VisualStudio.Shell.Interop;
+using System;
+using System.ComponentModel;
+using System.Runtime.InteropServices;
+
+namespace UnrealVS
+{
+	// GUID for this class.
+	[Guid(GuidList.UnrealVSPackageManString)]
+
+	class ProfileManager : Component, IProfileManager
+	{
+		public void SaveSettingsToXml(IVsSettingsWriter writer)
+		{
+			throw new NotImplementedException();
+		}
+
+		public void LoadSettingsFromXml(IVsSettingsReader reader)
+		{
+			throw new NotImplementedException();
+		}
+
+		public void SaveSettingsToStorage()
+		{
+			throw new NotImplementedException();
+		}
+
+		public void LoadSettingsFromStorage()
+		{
+			throw new NotImplementedException();
+		}
+
+		public void ResetSettings()
+		{
+			throw new NotImplementedException();
+		}
+	}
+}
diff --git a/Engine/Source/Programs/UnrealVS/UnrealVS.2019/Properties/AssemblyInfo.cs b/Engine/Source/Programs/UnrealVS/UnrealVS.2019/Properties/AssemblyInfo.cs
new file mode 100644
--- /dev/null
+++ b/Engine/Source/Programs/UnrealVS/UnrealVS.2019/Properties/AssemblyInfo.cs
@@ -0,0 +1,37 @@
+// Copyright Epic Games, Inc. All Rights Reserved.
+
+using System;
+using System.Reflection;
+using System.Resources;
+using System.Runtime.InteropServices;
+
+// General Information about an assembly is controlled through the following 
+// set of attributes. Change these attribute values to modify the information
+// associated with an assembly.
+[assembly: AssemblyTitle("UnrealVS")]
+[assembly: AssemblyDescription("")]
+[assembly: AssemblyConfiguration("")]
+[assembly: AssemblyCompany("Epic Games, Inc.")]
+[assembly: AssemblyProduct("UnrealVS")]
+[assembly: AssemblyCopyright("Copyright Epic Games, Inc.")]
+[assembly: AssemblyTrademark("")]
+[assembly: AssemblyCulture("")]
+[assembly: ComVisible(false)]
+[assembly: CLSCompliant(false)]
+[assembly: NeutralResourcesLanguage("en-US")]
+
+// Version information for an assembly consists of the following four values:
+//
+//      Major Version
+//      Minor Version 
+//      Build Number
+//      Revision
+//
+// You can specify all the values or you can default the Revision and Build Numbers 
+// by using the '*' as shown below:
+
+[assembly: AssemblyVersion("1.61.0.0")]
+[assembly: AssemblyFileVersion("1.61.0.0")]
+
+
+
diff --git a/Engine/Source/Programs/UnrealVS/UnrealVS.2019/ProvideSolutionProperties.cs b/Engine/Source/Programs/UnrealVS/UnrealVS.2019/ProvideSolutionProperties.cs
new file mode 100644
--- /dev/null
+++ b/Engine/Source/Programs/UnrealVS/UnrealVS.2019/ProvideSolutionProperties.cs
@@ -0,0 +1,50 @@
+// Copyright Epic Games, Inc. All Rights Reserved.
+
+using Microsoft.VisualStudio.Shell;
+using System;
+using System.Globalization;
+
+namespace UnrealVS
+{
+	/// <summary>
+	/// This attribute registers the package as a solution property parser
+	/// </summary>
+	[AttributeUsage(AttributeTargets.Class, AllowMultiple = true, Inherited = true)]
+	internal sealed class ProvideSolutionPropertiesAttribute : RegistrationAttribute
+	{
+		private readonly string _propName;
+
+		public ProvideSolutionPropertiesAttribute(string propName)
+		{
+			_propName = propName;
+		}
+
+		public override void Register(RegistrationContext context)
+		{
+			context.Log.WriteLine(string.Format(CultureInfo.InvariantCulture, "ProvideSolutionProps: ({0} = {1})", context.ComponentType.GUID.ToString("B"), PropName));
+
+			Key childKey = null;
+
+			try
+			{
+				childKey = context.CreateKey(string.Format(CultureInfo.InvariantCulture, "{0}\\{1}", "SolutionPersistence", PropName));
+
+				childKey.SetValue(string.Empty, context.ComponentType.GUID.ToString("B").ToUpperInvariant());
+			}
+			finally
+			{
+				if (childKey != null) childKey.Close();
+			}
+		}
+
+		public override void Unregister(RegistrationContext context)
+		{
+			context.RemoveKey(string.Format(CultureInfo.InvariantCulture, "{0}\\{1}", "SolutionPersistence", PropName));
+		}
+
+		public string PropName
+		{
+			get { return _propName; }
+		}
+	}
+}
diff --git a/Engine/Source/Programs/UnrealVS/UnrealVS.2019/QuickBuild.cs b/Engine/Source/Programs/UnrealVS/UnrealVS.2019/QuickBuild.cs
new file mode 100644
--- /dev/null
+++ b/Engine/Source/Programs/UnrealVS/UnrealVS.2019/QuickBuild.cs
@@ -0,0 +1,332 @@
+// Copyright Epic Games, Inc. All Rights Reserved.
+
+using EnvDTE;
+using EnvDTE80;
+using Microsoft.VisualStudio;
+using Microsoft.VisualStudio.Shell;
+using Microsoft.VisualStudio.Shell.Interop;
+using System;
+using System.Collections.Generic;
+using System.ComponentModel.Design;
+using System.Linq;
+using System.Runtime.InteropServices;
+
+namespace UnrealVS
+{
+	/// <summary>
+	/// Ole Command class for handling the quick build menu items (individual configs inside platform menus)
+	/// One of these handles all the items in a platform menu.
+	/// </summary>
+	class QuickBuildMenuCommand : OleMenuCommand
+	{
+		/// The function used to select the range of commands handled by this object
+		private readonly Predicate<int> OnMatchItem;
+
+		/// <summary>
+		/// Create a QuickBuildMenuCommand with a delegate for determining which command numbers it handles
+		/// The other params are used to construct the base OleMenuCommand
+		/// </summary>
+		public QuickBuildMenuCommand(CommandID InDynamicStartCommandId, Predicate<int> InOnMatchItem, EventHandler InOnInvokedDynamicItem, EventHandler InOnBeforeQueryStatusDynamicItem)
+			: base(InOnInvokedDynamicItem, null /*changeHandler*/, InOnBeforeQueryStatusDynamicItem, InDynamicStartCommandId)
+		{
+			OnMatchItem = InOnMatchItem ?? throw new ArgumentNullException("InOnMatchItem");
+		}
+
+		/// <summary>
+		/// Allows a dynamic item command to match the subsequent items in its list.
+		/// Overridden from OleMenuCommand.
+		/// Calls OnMatchItem to choose with command numbers are valid.
+		/// </summary>
+		public override bool DynamicItemMatch(int cmdId)
+		{
+			// Call the supplied predicate to test whether the given cmdId is a match.
+			// If it is, store the command id in MatchedCommandid
+			// for use by any BeforeQueryStatus handlers, and then return that it is a match.
+			// Otherwise clear any previously stored matched cmdId and return that it is not a match.
+			if (OnMatchItem(cmdId))
+			{
+				MatchedCommandId = cmdId;
+				return true;
+			}
+			MatchedCommandId = 0;
+			return false;
+		}
+	}
+
+	/// <summary>
+	/// Class represents a single configuration quick build menu. Wraps QuickBuildMenuCommand.
+	/// </summary>
+	class ConfigurationMenuContents
+	{
+		/// The command id of the placeholder command that sets the start of the command range for this menu
+		private readonly int DynamicStartCommandId;
+		/// The name of the configuration. Matches the config string in the solution configs exactly.
+		private readonly string ConfigurationName;
+
+		/// <summary>
+		/// Construct a platform menu for a given, named platform, starting at the given command id
+		/// </summary>
+		/// <param name="Name">The name of the platform</param>
+		/// <param name="InDynamicStartCommandId">The start of the command range for this menu</param>
+		public ConfigurationMenuContents(string Name, int InDynamicStartCommandId)
+		{
+			DynamicStartCommandId = InDynamicStartCommandId;
+			ConfigurationName = Name;
+
+			// Create the dynamic menu
+			var MenuOleCommand = new QuickBuildMenuCommand(
+				new CommandID(GuidList.UnrealVSCmdSet, DynamicStartCommandId),
+				IsValidDynamicItem,
+				OnInvokedDynamicItem,
+				OnBeforeQueryStatusDynamicItem
+			);
+			UnrealVSPackage.Instance.MenuCommandService.AddCommand(MenuOleCommand);
+		}
+
+		/// <summary>
+		/// The callback function passed to this object's QuickBuildMenuCommand, used to select the range of commands handled by this object
+		/// It uses the cached list of solution configs in QuickBuild to determine how many commands should be in the menu.
+		/// </summary>
+		/// <param name="cmdId">Command id the validate</param>
+		/// <returns>True, if the command is valid for this menu</returns>
+		private bool IsValidDynamicItem(int cmdId)
+		{
+			int ConfigCount = QuickBuild.CachedSolutionConfigPlatforms.Length;
+
+			// The match is valid if the command ID is >= the id of our root dynamic start item
+			// and the command ID minus the ID of our root dynamic start item
+			// is less than or equal to the number of projects in the solution.
+			return (cmdId >= DynamicStartCommandId) && ((cmdId - DynamicStartCommandId) < ConfigCount);
+		}
+
+		/// <summary>
+		/// The invoke handler passed to this object's QuickBuildMenuCommand, called by the base OleMenuCommand when an item is clicked
+		/// </summary>
+		private void OnInvokedDynamicItem(object sender, EventArgs args)
+		{
+			var MenuCommand = (QuickBuildMenuCommand)sender;
+
+			// Get the project clicked in the solution explorer by accessing the current selection and converting to a Project if possible.
+			UnrealVSPackage.Instance.SelectionManager.GetCurrentSelection(out IntPtr HierarchyPtr, out uint ProjectItemId, out IVsMultiItemSelect MultiItemSelect, out IntPtr SelectionContainerPtr);
+			if (HierarchyPtr == null) return;
+
+			if (!(Marshal.GetTypedObjectForIUnknown(HierarchyPtr, typeof(IVsHierarchy)) is IVsHierarchy SelectedHierarchy)) return;
+
+			var SelectedProject = Utils.HierarchyObjectToProject(SelectedHierarchy);
+			if (SelectedProject == null) return;
+
+			// Builds the selected project with the clicked platform and config
+			Utils.ExecuteProjectBuild(SelectedProject, ConfigurationName, MenuCommand.Text, BatchBuilderToolControl.BuildJob.BuildJobType.Build, null, null);
+		}
+
+		/// <summary>
+		/// Before-query handler passed to this object's QuickBuildMenuCommand, called by the base OleMenuCommand to update the state of an item.
+		/// </summary>
+		private void OnBeforeQueryStatusDynamicItem(object sender, EventArgs args)
+		{
+			var MenuCommand = (QuickBuildMenuCommand)sender;
+			MenuCommand.Enabled = QuickBuild.IsActive;
+			MenuCommand.Visible = true;
+
+			// Determine the index of the item in the menu
+			bool isRootItem = MenuCommand.MatchedCommandId == 0;
+			int CommandId = isRootItem ? DynamicStartCommandId : MenuCommand.MatchedCommandId;
+			int DynCommandIdx = CommandId - DynamicStartCommandId;
+
+			// Set the text label based on the index and the solution platforms array cached in QuickBuild
+			if (DynCommandIdx < QuickBuild.CachedSolutionConfigPlatforms.Length)
+			{
+				MenuCommand.Text = QuickBuild.CachedSolutionConfigPlatforms[DynCommandIdx];
+			}
+
+			// Clear the id now that the query is done
+			MenuCommand.MatchedCommandId = 0;
+		}
+	}
+
+	/// <summary>
+	/// The Quick Build menu system
+	/// </summary>
+	public class QuickBuild
+	{
+		/** classes */
+
+		/// All the basic data needed to build each platform-specific submenu
+		struct SubMenu
+		{
+			public string Name { get; set; }
+			public int SubMenuId { get; set; }
+			public int DynamicStartCommandId { get; set; }
+		}
+
+		/** constants */
+
+		private const int ProjectQuickBuildMenuID = 0x1410;     // must match the values in the vsct file
+
+		/** fields */
+
+		/// Solution configs and platforms cached once before each time the Quick Build menu opens
+		private static string[] SolutionConfigNames = new string[0];
+		private static string[] SolutionConfigPlatforms = new string[0];
+
+		/// Hide/shows the while Quick Build menu tree
+		private static bool bIsActive = false;
+
+		/// List of submenus and their details - must match the values in the vsct file
+		private readonly SubMenu[] SubMenus = new[]
+		{
+			new SubMenu {Name = "Debug", SubMenuId = 0x1500, DynamicStartCommandId = 0x1530},
+			new SubMenu {Name = "Debug Client", SubMenuId = 0x1600, DynamicStartCommandId = 0x1630},
+			new SubMenu {Name = "Debug Editor", SubMenuId = 0x1700, DynamicStartCommandId = 0x1730},
+			new SubMenu {Name = "Debug Server", SubMenuId = 0x1800, DynamicStartCommandId = 0x1830},
+			new SubMenu {Name = "DebugGame", SubMenuId = 0x1900, DynamicStartCommandId = 0x1930},
+			new SubMenu {Name = "DebugGame Client", SubMenuId = 0x1A00, DynamicStartCommandId = 0x1A30},
+			new SubMenu {Name = "DebugGame Editor", SubMenuId = 0x1B00, DynamicStartCommandId = 0x1B30},
+			new SubMenu {Name = "DebugGame Server", SubMenuId = 0x1C00, DynamicStartCommandId = 0x1C30},
+			new SubMenu {Name = "Development", SubMenuId = 0x1D00, DynamicStartCommandId = 0x1D30},
+			new SubMenu {Name = "Development Client", SubMenuId = 0x1E00, DynamicStartCommandId = 0x1E30},
+			new SubMenu {Name = "Development Editor", SubMenuId = 0x1F00, DynamicStartCommandId = 0x1F30},
+			new SubMenu {Name = "Development Server", SubMenuId = 0x2000, DynamicStartCommandId = 0x2030},
+			new SubMenu {Name = "Shipping", SubMenuId = 0x2100, DynamicStartCommandId = 0x2130},
+			new SubMenu {Name = "Shipping Client", SubMenuId = 0x2200, DynamicStartCommandId = 0x2230},
+			new SubMenu {Name = "Shipping Editor", SubMenuId = 0x2300, DynamicStartCommandId = 0x2330},
+			new SubMenu {Name = "Shipping Server", SubMenuId = 0x2400, DynamicStartCommandId = 0x2430},
+			new SubMenu {Name = "Test", SubMenuId = 0x2500, DynamicStartCommandId = 0x2530},
+			new SubMenu {Name = "Test Client", SubMenuId = 0x2600, DynamicStartCommandId = 0x2630},
+			new SubMenu {Name = "Test Editor", SubMenuId = 0x2700, DynamicStartCommandId = 0x2730},
+			new SubMenu {Name = "Test Server", SubMenuId = 0x2800, DynamicStartCommandId = 0x2830},
+
+		};
+
+		/// The main root command of the Quick Build menu hierarchy - used to hide it when not active
+		private readonly OleMenuCommand QuickBuildCommand;
+
+		/// <summary>
+		/// These represent the items that can be added to the menu that lists the platforms.
+		/// Each one is a submenu containing items for each config.
+		/// </summary>
+		private readonly Dictionary<string, OleMenuCommand> AllConfigurationMenus = new Dictionary<string, OleMenuCommand>();
+
+		/// <summary>
+		/// These represent the items shown in the menu that lists the configurations.
+		/// It is a subset of AllConfigurationMenus with only the loaded platforms.
+		/// Each one is a submenu containing items for each config.
+		/// </summary>
+		private readonly Dictionary<string, OleMenuCommand> ActiveConfigurationMenus = new Dictionary<string, OleMenuCommand>();
+
+		/// <summary>
+		/// These holds the objects that generate the per platform submenus for each configuration.
+		/// </summary>
+		// ReSharper disable once CollectionNeverQueried.Local - This is only invoked by visual studio to populate the submenus.
+		private readonly Dictionary<string, ConfigurationMenuContents> ConfigurationMenusContents = new Dictionary<string, ConfigurationMenuContents>();
+
+		/// VSConstants.UICONTEXT_SolutionBuilding translated into a cookie used to access UI ctxt state
+		private readonly uint SolutionBuildingUIContextCookie;
+
+		/** properties */
+
+		public static bool IsActive { get { return bIsActive; } }
+		public static string[] CachedSolutionConfigPlatforms { get { return SolutionConfigPlatforms; } }
+
+		/** methods */
+
+		public QuickBuild()
+		{
+			// root menu
+			QuickBuildCommand = new OleMenuCommand(null, null, OnQuickBuildQuery, new CommandID(GuidList.UnrealVSCmdSet, ProjectQuickBuildMenuID));
+			QuickBuildCommand.BeforeQueryStatus += OnQuickBuildQuery;
+			UnrealVSPackage.Instance.MenuCommandService.AddCommand(QuickBuildCommand);
+
+			// configuration sub-menus
+			foreach (var SubMenu in SubMenus)
+			{
+				var SubMenuCommand = new OleMenuCommand(null, new CommandID(GuidList.UnrealVSCmdSet, SubMenu.SubMenuId));
+				SubMenuCommand.BeforeQueryStatus += OnQuickBuildSubMenuQuery;
+				UnrealVSPackage.Instance.MenuCommandService.AddCommand(SubMenuCommand);
+				AllConfigurationMenus.Add(SubMenu.Name, SubMenuCommand);
+				ConfigurationMenusContents.Add(SubMenu.Name, new ConfigurationMenuContents(SubMenu.Name, SubMenu.DynamicStartCommandId));
+			}
+			// cache the cookie for UICONTEXT_SolutionBuilding
+			UnrealVSPackage.Instance.SelectionManager.GetCmdUIContextCookie(VSConstants.UICONTEXT_SolutionBuilding, out SolutionBuildingUIContextCookie);
+
+			// Initialize the active state based on whether the IDE is building anything
+			UpdateActiveState();
+
+			UnrealVSPackage.Instance.OnUIContextChanged += OnUIContextChanged;
+			UnrealVSPackage.Instance.OnSolutionOpened += OnSolutionChanged;
+
+			OnSolutionChanged();
+		}
+
+		/// <summary>
+		/// Updates the bIsActive flag using the UICONTEXT_SolutionBuilding state. Hides the Quick Build feature when the solution is building.
+		/// </summary>
+		private void UpdateActiveState()
+		{
+			UnrealVSPackage.Instance.SelectionManager.IsCmdUIContextActive(SolutionBuildingUIContextCookie, out int bIsBuilding);
+			bIsActive = bIsBuilding == 0;
+		}
+
+		/// <summary>
+		/// Before-query handler passed to the root menu item's OleMenuCommand and called to update the state of the item.
+		/// </summary>
+		private void OnQuickBuildQuery(object sender, EventArgs e)
+		{
+			// Always cache the list of solution build configs when the project menu is opening
+			CacheBuildConfigs();
+		}
+
+		/// <summary>
+		/// Before-query handler passed to the sub-menu item's OleMenuCommand and called to update the state of the item.
+		/// </summary>
+		private void OnQuickBuildSubMenuQuery(object sender, EventArgs e)
+		{
+			var SubMenuCommand = (OleMenuCommand)sender;
+			SubMenuCommand.Visible = ActiveConfigurationMenus.ContainsValue(SubMenuCommand);
+			SubMenuCommand.Enabled = bIsActive;
+		}
+
+		/// <summary>
+		/// Called when a solution loads. Caches the solution build configs and sets the platform menus' visibility.
+		/// Only configs found in the loaded solution's list are shown.
+		/// </summary>
+		private void OnSolutionChanged()
+		{
+			CacheBuildConfigs();
+
+			ActiveConfigurationMenus.Clear();
+			Logging.WriteLine("Updating solution menu state, currently active configs are: " + string.Join(", ", SolutionConfigNames));
+			foreach (var SubMenu in SubMenus)
+			{
+				if (SolutionConfigNames.Any(Config => string.Compare(Config, SubMenu.Name, StringComparison.InvariantCultureIgnoreCase) == 0))
+				{
+					ActiveConfigurationMenus.Add(SubMenu.Name, AllConfigurationMenus[SubMenu.Name]);
+				}
+			}
+		}
+
+		/// <summary>
+		/// Caches the solution build configs
+		/// </summary>
+		private void CacheBuildConfigs()
+		{
+			SolutionConfigurations SolutionConfigs = UnrealVSPackage.Instance.DTE.Solution.SolutionBuild.SolutionConfigurations;
+
+			SolutionConfigPlatforms = (from SolutionConfiguration2 Sc in SolutionConfigs select Sc.PlatformName).Distinct().ToArray();
+
+			SolutionConfigNames = (from SolutionConfiguration2 Sc in SolutionConfigs select Sc.Name).Distinct().ToArray();
+
+		}
+
+		/// <summary>
+		/// Called when the UI Context changes. If the building state changed, update the active state.
+		/// </summary>
+		private void OnUIContextChanged(uint CmdUICookie, bool bActive)
+		{
+			if (SolutionBuildingUIContextCookie == CmdUICookie)
+			{
+				UpdateActiveState();
+			}
+		}
+	}
+}
diff --git a/Engine/Source/Programs/UnrealVS/UnrealVS.2019/StartupProjectSelector.cs b/Engine/Source/Programs/UnrealVS/UnrealVS.2019/StartupProjectSelector.cs
new file mode 100644
--- /dev/null
+++ b/Engine/Source/Programs/UnrealVS/UnrealVS.2019/StartupProjectSelector.cs
@@ -0,0 +1,293 @@
+// Copyright Epic Games, Inc. All Rights Reserved.
+
+using EnvDTE;
+using Microsoft.VisualStudio.Shell;
+using Microsoft.VisualStudio.Shell.Interop;
+using System;
+using System.Collections.Generic;
+using System.ComponentModel.Design;
+using System.Linq;
+using System.Runtime.InteropServices;
+
+namespace UnrealVS
+{
+	public class StartupProjectSelector
+	{
+		const int StartupProjectComboId = 0x1050;
+		const int StartupProjectComboListId = 0x1060;
+
+		public class StartupProjectListChangedEventArgs : EventArgs
+		{
+			public Project[] StartupProjects { get; set; }
+		}
+		public delegate void StartupProjectListChangedEventHandler(object sender, StartupProjectListChangedEventArgs e);
+		public event StartupProjectListChangedEventHandler StartupProjectListChanged;
+
+		public bool? CachedHideNonGameStartupProjects { get; set; }
+
+		private class ProjectReference
+		{
+			public Project Project { get; set; }
+
+			public string Name
+			{
+				get { return (Project != null ? Project.Name : "<invalid>"); }
+			}
+
+			public override string ToString()
+			{
+				return (Project != null ? Project.Name : "<invalid>");
+			}
+		}
+
+		public StartupProjectSelector()
+		{
+			_bIsSolutionOpened = UnrealVSPackage.Instance.DTE.Solution.IsOpen;
+
+			// Create the handlers for our commands
+			{
+				// StartupProjectCombo
+				var StartupProjectComboCommandId = new CommandID(GuidList.UnrealVSCmdSet, StartupProjectComboId);
+				var StartupProjectComboCommand = new OleMenuCommand(StartupProjectComboHandler, StartupProjectComboCommandId);
+				StartupProjectComboCommand.BeforeQueryStatus += (sender, args) => StartupProjectComboCommand.Enabled = _bIsSolutionOpened;
+				UnrealVSPackage.Instance.MenuCommandService.AddCommand(StartupProjectComboCommand);
+
+				// StartupProjectComboList
+				var StartupProjectComboListCommandId = new CommandID(GuidList.UnrealVSCmdSet, StartupProjectComboListId);
+				var StartupProjectComboListCommand = new OleMenuCommand(StartupProjectComboListHandler, StartupProjectComboListCommandId);
+				StartupProjectComboListCommand.BeforeQueryStatus += (sender, args) => StartupProjectComboListCommand.Enabled = _bIsSolutionOpened;
+				UnrealVSPackage.Instance.MenuCommandService.AddCommand(StartupProjectComboListCommand);
+
+				StartupProjectComboCommands = new[] { StartupProjectComboCommand, StartupProjectComboListCommand };
+			}
+
+			// Register for events that we care about
+			UnrealVSPackage.Instance.OnStartupProjectChanged += OnStartupProjectChanged;
+			UnrealVSPackage.Instance.OnSolutionOpened +=
+				delegate
+				{
+					Logging.WriteLine("Opened solution " + UnrealVSPackage.Instance.DTE.Solution.FullName);
+					_bIsSolutionOpened = true;
+					UpdateStartupProjectCombo();
+					UpdateStartupProjectList(Utils.GetAllProjectsFromDTE());
+				};
+			UnrealVSPackage.Instance.OnSolutionClosing +=
+				delegate
+				{
+					Logging.WriteLine("Closing solution");
+					_bIsSolutionOpened = false;
+					UpdateStartupProjectCombo();
+					_CachedStartupProjects.Clear();
+				};
+			UnrealVSPackage.Instance.OnProjectOpened +=
+				delegate (Project OpenedProject)
+				{
+					if (_bIsSolutionOpened)
+					{
+						Logging.WriteLine("Opened project node " + OpenedProject.Name);
+						UpdateStartupProjectList(OpenedProject);
+					}
+					else
+					{
+						Logging.WriteLine("Opened project node " + OpenedProject.Name + " with the solution CLOSED");
+					}
+				};
+			UnrealVSPackage.Instance.OnProjectClosed +=
+				delegate (Project ClosedProject)
+				{
+					Logging.WriteLine("Closed project node " + ClosedProject.Name);
+					RemoveFromStartupProjectList(ClosedProject);
+				};
+			UnrealVSPackage.Instance.OptionsPage.OnOptionsChanged += OnOptionsChanged;
+
+			UpdateStartupProjectList(Utils.GetAllProjectsFromDTE());
+			UpdateStartupProjectCombo();
+		}
+
+		public Project[] StartupProjectList
+		{
+			get
+			{
+				return (from ProjectRef in _CachedStartupProjects select ProjectRef.Project).ToArray();
+			}
+		}
+
+		private void OnOptionsChanged(object Sender, EventArgs E)
+		{
+			UpdateCachedOptions();
+
+			UpdateStartupProjectList(Utils.GetAllProjectsFromDTE());
+		}
+
+		private void UpdateCachedOptions()
+		{
+			CachedHideNonGameStartupProjects = UnrealVSPackage.Instance.OptionsPage.HideNonGameStartupProjects;
+		}
+
+		/// Rebuild/update the list CachedStartupProjectNames whenever something changes
+		private void UpdateStartupProjectList(params Project[] DirtyProjects)
+		{
+			foreach (Project Project in DirtyProjects)
+			{
+				if (ProjectPredicate(Project))
+				{
+					if (!_CachedStartupProjects.Any(ProjRef => ProjRef.Project == Project))
+					{
+						_CachedStartupProjects.Add(new ProjectReference { Project = Project });
+					}
+				}
+				else
+				{
+					RemoveFromStartupProjectList(Project);
+				}
+			}
+
+			SortCachedStartupProjects();
+
+			StartupProjectListChanged?.Invoke(this, new StartupProjectListChangedEventArgs { StartupProjects = StartupProjectList });
+		}
+
+		private void SortCachedStartupProjects()
+		{
+			// Sort projects by game y/n then alphabetically
+			_CachedStartupProjects = (_CachedStartupProjects.OrderBy(ProjectRef => ProjectRef.Name == "UE5" ? 0 : 1)
+				.ThenBy(ProjectRef => Utils.IsGameProject(ProjectRef.Project) ? 0 : 1)
+				.ThenBy(ProjectRef => ProjectRef.Name)).ToList();
+		}
+
+		private bool ProjectPredicate(Project Project)
+		{
+			// Initialize the cached options if they haven't been read from UnrealVSOptions yet
+			if (!CachedHideNonGameStartupProjects.HasValue)
+			{
+				UpdateCachedOptions();
+			}
+
+			// Optionally, ignore non-game projects
+			if (CachedHideNonGameStartupProjects.Value)
+			{
+				if (!Utils.IsGameProject(Project))
+				{
+					Logging.WriteLine("StartupProjectSelector: Not listing project " + Project.Name + " because it is not a game");
+					return false;
+				}
+			}
+
+			// Always filter out non-executable projects
+			if (!Utils.IsProjectBuildable(Project))
+			{
+				Logging.WriteLine("StartupProjectSelector: Not listing project " + Project.Name + " because it is not executable");
+				return false;
+			}
+
+			Logging.WriteLine("StartupProjectSelector: Listing project " + Project.Name);
+			return true;
+		}
+
+		/// Remove projects from the list CachedStartupProjectNames whenever one unloads
+		private void RemoveFromStartupProjectList(Project Project)
+		{
+			_CachedStartupProjects.RemoveAll(ProjRef => ProjRef.Project == Project);
+
+			StartupProjectListChanged?.Invoke(this, new StartupProjectListChangedEventArgs { StartupProjects = StartupProjectList });
+		}
+
+		/// <summary>
+		/// Updates the combo box after the project loaded state has changed
+		/// </summary>
+		private void UpdateStartupProjectCombo()
+		{
+			foreach (var Command in StartupProjectComboCommands)
+			{
+				Command.Enabled = _bIsSolutionOpened;
+			}
+		}
+
+
+		/// <summary>
+		/// Called when the startup project has changed to a new project.  We'll refresh our interface.
+		/// </summary>
+		public void OnStartupProjectChanged(Project NewStartupProject)
+		{
+			UpdateStartupProjectCombo();
+		}
+
+
+		/// <summary>
+		/// Returns a string to display in the startup project combo box, based on project state
+		/// </summary>
+		/// <returns>String to display</returns>
+		private string MakeStartupProjectComboText()
+		{
+			string Text = "";
+
+			// Switch to this project!
+			UnrealVSPackage.Instance.SolutionBuildManager.get_StartupProject(out IVsHierarchy ProjectHierarchy);
+			if (ProjectHierarchy != null)
+			{
+				var Project = Utils.HierarchyObjectToProject(ProjectHierarchy);
+				if (Project != null)
+				{
+					Text = Project.Name;
+				}
+			}
+
+			return Text;
+		}
+
+
+		/// Called by combo control to query the text to display or to apply newly-entered text
+		void StartupProjectComboHandler(object Sender, EventArgs Args)
+		{
+			var OleArgs = (OleMenuCmdEventArgs)Args;
+
+			// If OutValue is non-zero, then Visual Studio is querying the current combo value to display
+			if (OleArgs.OutValue != IntPtr.Zero)
+			{
+				string ValueString = MakeStartupProjectComboText();
+				Marshal.GetNativeVariantForObject(ValueString, OleArgs.OutValue);
+			}
+
+			// Has a new input value been entered?
+			else if (OleArgs.InValue != null)
+			{
+				var InputString = OleArgs.InValue.ToString();
+
+				// Lookup the ProjectReference in CachedStartupProjects to get the full name
+				ProjectReference ProjectRefMatch = _CachedStartupProjects.FirstOrDefault(ProjRef => ProjRef.Name == InputString);
+
+				if (ProjectRefMatch != null && ProjectRefMatch.Project != null)
+				{
+					// Switch to this project using the standard SDK method set_StartupProject
+					var ProjectHierarchy = Utils.ProjectToHierarchyObject(ProjectRefMatch.Project);
+					UnrealVSPackage.Instance.SolutionBuildManager.set_StartupProject(ProjectHierarchy);
+
+					// Select in solution exp and set as startup project
+					// This forces the .suo to save the project change
+					if (Utils.SelectProjectInSolutionExplorer(ProjectRefMatch.Project))
+					{
+						UnrealVSPackage.Instance.DTE.ExecuteCommand("Project.SetasStartUpProject");
+					}
+				}
+			}
+		}
+
+		/// Called by combo control to populate the drop-down list
+		void StartupProjectComboListHandler(object Sender, EventArgs Args)
+		{
+			var OleArgs = (OleMenuCmdEventArgs)Args;
+
+			var CachedStartupProjectNames = (from ProjRef in _CachedStartupProjects select ProjRef.Name).ToArray();
+			Marshal.GetNativeVariantForObject(CachedStartupProjectNames, OleArgs.OutValue);
+		}
+
+		/// Combo control and list commands for our startup project selector
+		readonly OleMenuCommand[] StartupProjectComboCommands;
+
+		/// Keep track of whether a solution is opened
+		bool _bIsSolutionOpened;
+
+		/// Store the list of projects that should be shown in the dropdown list
+		List<ProjectReference> _CachedStartupProjects = new List<ProjectReference>();
+	}
+}
diff --git a/Engine/Source/Programs/UnrealVS/UnrealVS.2019/UnrealVS.2019.csproj b/Engine/Source/Programs/UnrealVS/UnrealVS.2019/UnrealVS.2019.csproj
new file mode 100644
--- /dev/null
+++ b/Engine/Source/Programs/UnrealVS/UnrealVS.2019/UnrealVS.2019.csproj
@@ -0,0 +1,251 @@
+﻿<?xml version="1.0" encoding="utf-8"?>
+<Project DefaultTargets="Build" xmlns="http://schemas.microsoft.com/developer/msbuild/2003" ToolsVersion="14.0">
+  <Import Project="$(MSBuildExtensionsPath)\$(MSBuildToolsVersion)\Microsoft.Common.props" Condition="Exists('$(MSBuildExtensionsPath)\$(MSBuildToolsVersion)\Microsoft.Common.props')" />
+  <PropertyGroup>
+    <Configuration Condition=" '$(Configuration)' == '' ">Debug</Configuration>
+    <Platform Condition=" '$(Platform)' == '' ">AnyCPU</Platform>
+    <SchemaVersion>2.0</SchemaVersion>
+    <ProjectGuid>{C0EE3E24-0255-4D0E-9369-82247AD5E2DD}</ProjectGuid>
+    <ProjectTypeGuids>{82b43b9b-a64c-4715-b499-d71e9ca2bd60};{60dc8134-eba5-43b8-bcc9-bb4bc16c2548};{FAE04EC0-301F-11D3-BF4B-00C04F79EFBC}</ProjectTypeGuids>
+    <OutputType>Library</OutputType>
+    <AppDesignerFolder>Properties</AppDesignerFolder>
+    <RootNamespace>UnrealVS</RootNamespace>
+    <AssemblyName>UnrealVS</AssemblyName>
+    <SignAssembly>True</SignAssembly>
+    <AssemblyOriginatorKeyFile>Key.snk</AssemblyOriginatorKeyFile>
+    <TargetFrameworkVersion>v4.7.2</TargetFrameworkVersion>
+    <SccProjectName>
+    </SccProjectName>
+    <SccLocalPath>
+    </SccLocalPath>
+    <SccAuxPath>
+    </SccAuxPath>
+    <SccProvider>
+    </SccProvider>
+    <MinimumVisualStudioVersion>15.0</MinimumVisualStudioVersion>
+    <OldToolsVersion>14.0</OldToolsVersion>
+    <TargetFrameworkProfile />
+    <FileUpgradeFlags>
+    </FileUpgradeFlags>
+    <UpgradeBackupLocation>
+    </UpgradeBackupLocation>
+    <IsWebBootstrapper>false</IsWebBootstrapper>
+    <PublishUrl>publish\</PublishUrl>
+    <Install>true</Install>
+    <InstallFrom>Disk</InstallFrom>
+    <UpdateEnabled>false</UpdateEnabled>
+    <UpdateMode>Foreground</UpdateMode>
+    <UpdateInterval>7</UpdateInterval>
+    <UpdateIntervalUnits>Days</UpdateIntervalUnits>
+    <UpdatePeriodically>false</UpdatePeriodically>
+    <UpdateRequired>false</UpdateRequired>
+    <MapFileExtensions>true</MapFileExtensions>
+    <ApplicationRevision>0</ApplicationRevision>
+    <ApplicationVersion>1.0.0.%2a</ApplicationVersion>
+    <UseApplicationTrust>false</UseApplicationTrust>
+    <BootstrapperEnabled>true</BootstrapperEnabled>
+  </PropertyGroup>
+  <ItemGroup>
+    <VSCTCompile Include="UnrealVS.vsct">
+      <ResourceName>Menus.ctmenu</ResourceName>
+      <SubType>Designer</SubType>
+    </VSCTCompile>
+  </ItemGroup>
+  <ItemGroup>
+    <Reference Include="PresentationCore" />
+    <Reference Include="PresentationFramework" />
+    <Reference Include="System" />
+    <Reference Include="System.Core" />
+    <Reference Include="System.Data" />
+    <Reference Include="System.Design" />
+    <Reference Include="System.Drawing" />
+    <Reference Include="System.Numerics" />
+    <Reference Include="System.Windows.Forms" />
+    <Reference Include="System.Xaml" />
+    <Reference Include="System.Xml" />
+    <Reference Include="WindowsBase" />
+  </ItemGroup>
+  <ItemGroup>
+    <Compile Include="BatchBuilder.cs" />
+    <Compile Include="BatchBuilderToolControl.xaml.cs">
+      <DependentUpon>BatchBuilderToolControl.xaml</DependentUpon>
+    </Compile>
+    <Compile Include="BatchBuilderToolWindow.cs" />
+    <Compile Include="BuildStartupProject.cs" />
+    <Compile Include="BuildStatusIcon.xaml.cs">
+      <DependentUpon>BuildStatusIcon.xaml</DependentUpon>
+    </Compile>
+    <Compile Include="CommandLineEditor.cs" />
+    <Compile Include="CompileSingleFile.cs" />
+    <Compile Include="GenerateProjectFiles.cs" />
+    <Compile Include="Guids.cs" />
+    <Compile Include="Logging.cs" />
+    <Compile Include="ProfileManager.cs">
+      <SubType>Component</SubType>
+    </Compile>
+    <Compile Include="ProvideSolutionProperties.cs" />
+    <Compile Include="QuickBuild.cs" />
+    <Compile Include="Resources.Designer.cs">
+      <AutoGen>True</AutoGen>
+      <DesignTime>True</DesignTime>
+      <DependentUpon>Resources.resx</DependentUpon>
+    </Compile>
+    <Compile Include="GlobalSuppressions.cs" />
+    <Compile Include="StartupProjectSelector.cs" />
+    <Compile Include="UnrealVSOptions.cs">
+      <SubType>Component</SubType>
+    </Compile>
+    <Compile Include="UnrealVSPackage.cs" />
+    <Compile Include="Properties\AssemblyInfo.cs" />
+    <Compile Include="Utils.cs" />
+  </ItemGroup>
+  <ItemGroup>
+    <EmbeddedResource Include="Resources.resx">
+      <Generator>ResXFileCodeGenerator</Generator>
+      <LastGenOutput>Resources.Designer.cs</LastGenOutput>
+      <SubType>Designer</SubType>
+    </EmbeddedResource>
+    <EmbeddedResource Include="VSPackage.resx">
+      <MergeWithCTO>true</MergeWithCTO>
+      <ManifestResourceName>VSPackage</ManifestResourceName>
+      <SubType>Designer</SubType>
+    </EmbeddedResource>
+  </ItemGroup>
+  <ItemGroup>
+    <None Include="..\..\..\..\Extras\UnrealVS\VS2019\UnrealVS.vsix">
+      <Link>UnrealVS.vsix</Link>
+    </None>
+    <None Include="Key.snk" />
+  </ItemGroup>
+  <ItemGroup>
+    <None Include="source.extension.vsixmanifest" />
+    <None Include="Resources\BatchBuilder.png" />
+    <Resource Include="Resources\BatchBuilderBackground.png" />
+    <Resource Include="Resources\Icon.ico" />
+    <Content Include="Resources\Icon.png">
+      <CopyToOutputDirectory>Always</CopyToOutputDirectory>
+      <IncludeInVSIX>true</IncludeInVSIX>
+    </Content>
+    <Content Include="Resources\Preview.png">
+      <IncludeInVSIX>true</IncludeInVSIX>
+      <CopyToOutputDirectory>Always</CopyToOutputDirectory>
+    </Content>
+    <Resource Include="Resources\RefreshProjects.png" />
+  </ItemGroup>
+  <ItemGroup>
+    <BootstrapperPackage Include=".NETFramework,Version=v4.5">
+      <Visible>False</Visible>
+      <ProductName>Microsoft .NET Framework 4.5 %28x86 and x64%29</ProductName>
+      <Install>true</Install>
+    </BootstrapperPackage>
+    <BootstrapperPackage Include="Microsoft.Net.Client.3.5">
+      <Visible>False</Visible>
+      <ProductName>.NET Framework 3.5 SP1 Client Profile</ProductName>
+      <Install>false</Install>
+    </BootstrapperPackage>
+    <BootstrapperPackage Include="Microsoft.Net.Framework.3.5.SP1">
+      <Visible>False</Visible>
+      <ProductName>.NET Framework 3.5 SP1</ProductName>
+      <Install>false</Install>
+    </BootstrapperPackage>
+  </ItemGroup>
+  <ItemGroup>
+    <Page Include="BatchBuilderToolControl.xaml">
+      <SubType>Designer</SubType>
+      <Generator>MSBuild:Compile</Generator>
+      <Generator>MSBuild:Compile</Generator>
+      <SubType>Designer</SubType>
+      <Generator>MSBuild:Compile</Generator>
+      <SubType>Designer</SubType>
+    </Page>
+    <Page Include="BuildStatusIcon.xaml">
+      <SubType>Designer</SubType>
+      <Generator>MSBuild:Compile</Generator>
+      <Generator>MSBuild:Compile</Generator>
+      <SubType>Designer</SubType>
+      <Generator>MSBuild:Compile</Generator>
+      <SubType>Designer</SubType>
+    </Page>
+  </ItemGroup>
+  <ItemGroup>
+    <PackageReference Include="Microsoft.CodeAnalysis.CSharp">
+      <Version>3.10.0</Version>
+    </PackageReference>
+    <PackageReference Include="Microsoft.CodeAnalysis.NetAnalyzers">
+      <Version>5.0.3</Version>
+      <IncludeAssets>runtime; build; native; contentfiles; analyzers; buildtransitive</IncludeAssets>
+      <PrivateAssets>all</PrivateAssets>
+    </PackageReference>
+    <PackageReference Include="Microsoft.CSharp">
+      <Version>4.7.0</Version>
+    </PackageReference>
+    <PackageReference Include="Microsoft.VisualStudio.Sdk">
+      <Version>16.10.31321.278</Version>
+    </PackageReference>
+    <PackageReference Include="Microsoft.VisualStudio.Threading">
+      <Version>16.10.56</Version>
+    </PackageReference>
+    <PackageReference Include="Microsoft.VisualStudio.Threading.Analyzers">
+      <Version>16.10.56</Version>
+      <IncludeAssets>runtime; build; native; contentfiles; analyzers; buildtransitive</IncludeAssets>
+      <PrivateAssets>all</PrivateAssets>
+    </PackageReference>
+    <PackageReference Include="Microsoft.VisualStudio.VCProjectEngine">
+      <Version>16.10.31320.204</Version>
+    </PackageReference>
+    <PackageReference Include="Microsoft.VSSDK.BuildTools">
+      <Version>17.0.1600</Version>
+      <IncludeAssets>runtime; build; native; contentfiles; analyzers; buildtransitive</IncludeAssets>
+      <PrivateAssets>all</PrivateAssets>
+    </PackageReference>
+    <PackageReference Include="System.Text.Json">
+      <Version>5.0.2</Version>
+    </PackageReference>
+  </ItemGroup>
+  <PropertyGroup>
+    <UseCodebase>true</UseCodebase>
+  </PropertyGroup>
+  <PropertyGroup>
+    <VisualStudioVersion Condition="'$(VisualStudioVersion)' == ''">11.0</VisualStudioVersion>
+    <VSToolsPath Condition="'$(VSToolsPath)' == ''">$(MSBuildExtensionsPath32)\Microsoft\VisualStudio\v$(VisualStudioVersion)</VSToolsPath>
+  </PropertyGroup>
+  <PropertyGroup Condition=" '$(Configuration)|$(Platform)' == 'Debug|AnyCPU' ">
+    <DebugSymbols>true</DebugSymbols>
+    <OutputPath>bin\Debug\</OutputPath>
+    <DefineConstants>TRACE;DEBUG;VS16</DefineConstants>
+    <DebugType>full</DebugType>
+    <PlatformTarget>AnyCPU</PlatformTarget>
+    <ErrorReport>prompt</ErrorReport>
+    <CodeAnalysisRuleSet>MinimumRecommendedRules.ruleset</CodeAnalysisRuleSet>
+    <NoWarn>MSB3277</NoWarn>
+  </PropertyGroup>
+  <PropertyGroup Condition=" '$(Configuration)|$(Platform)' == 'Release|AnyCPU' ">
+    <OutputPath>bin\Release\</OutputPath>
+    <DefineConstants>TRACE;VS16</DefineConstants>
+    <Optimize>true</Optimize>
+    <DebugType>pdbonly</DebugType>
+    <PlatformTarget>AnyCPU</PlatformTarget>
+    <RunCodeAnalysis>false</RunCodeAnalysis>
+    <ErrorReport>prompt</ErrorReport>
+    <CodeAnalysisRuleSet>MinimumRecommendedRules.ruleset</CodeAnalysisRuleSet>
+    <CopyVsixExtensionFiles>False</CopyVsixExtensionFiles>
+    <NoWarn>MSB3277</NoWarn>
+  </PropertyGroup>
+  <Import Project="$(MSBuildBinPath)\Microsoft.CSharp.targets" />
+  <Import Project="$(VSToolsPath)\VSSDK\Microsoft.VsSDK.targets" Condition="'$(VSToolsPath)' != ''" />
+  <Import Project="$(MSBuildExtensionsPath32)\Microsoft\VisualStudio\v11.0\VSSDK\Microsoft.VsSDK.targets" Condition="false" />
+  <PropertyGroup>
+    <PostBuildEvent>copy "$(TargetDir)$(TargetName).vsix" "$(ProjectDir)..\..\..\..\Extras\UnrealVS\VS2019\$(TargetName).vsix" /Y</PostBuildEvent>
+  </PropertyGroup>
+  <PropertyGroup>
+    <PreBuildEvent>
+    </PreBuildEvent>
+  </PropertyGroup>
+  <!-- To modify your build process, add your task inside one of the targets below and uncomment it. 
+       Other similar extension points exist, see Microsoft.Common.targets.
+  <Target Name="BeforeBuild">
+  </Target>
+  <Target Name="AfterBuild">
+  </Target>
+  -->
+</Project>
\ No newline at end of file
diff --git a/Engine/Source/Programs/UnrealVS/UnrealVS.2019/UnrealVSOptions.cs b/Engine/Source/Programs/UnrealVS/UnrealVS.2019/UnrealVSOptions.cs
new file mode 100644
--- /dev/null
+++ b/Engine/Source/Programs/UnrealVS/UnrealVS.2019/UnrealVSOptions.cs
@@ -0,0 +1,36 @@
+// Copyright Epic Games, Inc. All Rights Reserved.
+
+using Microsoft.VisualStudio.Shell;
+using System;
+using System.ComponentModel;
+using System.Runtime.InteropServices;
+
+namespace UnrealVS
+{
+	[Guid(GuidList.UnrealVSOptionsString)]
+	public class UnrealVsOptions : DialogPage
+	{
+		public event EventHandler OnOptionsChanged;
+
+		private bool _HideNonGameStartupProjects;
+
+		[Category("General")]
+		[DisplayName("Hide Non-Game Startup Projects")]
+		[Description("Shows only game projects in the startup project and batch-builder lists")]
+		public bool HideNonGameStartupProjects
+		{
+			get { return _HideNonGameStartupProjects; }
+			set { _HideNonGameStartupProjects = value; }
+		}
+
+		protected override void OnApply(PageApplyEventArgs e)
+		{
+			base.OnApply(e);
+
+			if (e.ApplyBehavior == ApplyKind.Apply && OnOptionsChanged != null)
+			{
+				OnOptionsChanged(this, EventArgs.Empty);
+			}
+		}
+	}
+}
diff --git a/Engine/Source/Programs/UnrealVS/UnrealVS.2019/UnrealVSPackage.cs b/Engine/Source/Programs/UnrealVS/UnrealVS.2019/UnrealVSPackage.cs
new file mode 100644
--- /dev/null
+++ b/Engine/Source/Programs/UnrealVS/UnrealVS.2019/UnrealVSPackage.cs
@@ -0,0 +1,970 @@
+// Copyright Epic Games, Inc. All Rights Reserved.
+
+using EnvDTE;
+using EnvDTE80;
+using Microsoft.VisualStudio;
+using Microsoft.VisualStudio.OLE.Interop;
+using Microsoft.VisualStudio.Shell;
+using Microsoft.VisualStudio.Shell.Interop;
+using System;
+using System.Collections.Generic;
+using System.ComponentModel.Design;
+using System.Diagnostics;
+using System.IO;
+using System.Linq;
+using System.Runtime.InteropServices;
+using System.Threading;
+using Thread = System.Threading.Thread;
+
+namespace UnrealVS
+{
+	public class UnrealSolutionProperties
+	{
+		public UnrealSolutionProperties(string[] InAvailablePlatforms)
+		{
+			AvailablePlatforms = InAvailablePlatforms;
+		}
+
+		public string[] AvailablePlatforms { get; }
+	}
+
+	// This attribute tells the PkgDef creation utility (CreatePkgDef.exe) that this class is a VS package.
+	[PackageRegistration(UseManagedResourcesOnly = true, AllowsBackgroundLoading = true)]
+
+	// This attribute is used to register the informations needed to show the this package in the Help/About dialog of Visual Studio.
+	[InstalledProductRegistration("#110", "#112", VersionString, IconResourceID = 400)]
+
+	// Adds data for our interface elements defined in UnrealVS.vsct.  The name "Menus.ctmenu" here is arbitrary, but must
+	// match the ItemGroup inside the UnrealVS.csproj MSBuild file (hand-typed.)
+	[ProvideMenuResource("Menus.ctmenu", 1)]
+
+	// Force the package to load whenever a solution exists
+	//[ProvideAutoLoad(VSConstants.UICONTEXT.SolutionExistsAndFullyLoaded_string, PackageAutoLoadFlags.BackgroundLoad)]
+	[ProvideAutoLoad(UIContextGuids80.SolutionExists, PackageAutoLoadFlags.BackgroundLoad)]
+
+	// Register the settings implementing class as providing support to UnrealVSPackage.
+	//[ProvideProfile(typeof (ProfileManager), "UnrealVS", "UnrealVSPackage", 110, 113, false)]
+
+	// GUID for this class.  This needs to match in quite a few places in the project.
+	[Guid(GuidList.UnrealVSPackageString)]
+
+	// This attribute registers a tool window exposed by this package.
+	[ProvideToolWindow(typeof(BatchBuilderToolWindow))]
+
+	// This attribute registers an options page for the package.
+	[ProvideOptionPage(typeof(UnrealVsOptions), ExtensionName, "General", 101, 102, true)]
+	[ProvideSolutionProperties(GuidList.UnrealVSPackageString)]
+	/// <summary>
+	/// UnrealVSPackage implements Package abstract class.  This is the main class that is registered
+	/// with Visual Studio shell and serves as the entry point into our extension
+	/// </summary>
+	public sealed class UnrealVSPackage :
+		AsyncPackage,           // We inherit from AsyncPackage which allows us to become a "plugin" within the Visual Studio shell
+		IVsSolutionEvents,      // This interface allows us to register to be notified of events such as opening a project
+		IVsUpdateSolutionEvents,// Allows us to register to be notified of events such as active config changes
+		IVsSelectionEvents,     // Allows us to be notified when the startup project has changed to a different project
+		IVsHierarchyEvents,     // Allows us to be notified when a hierarchy (the startup project) has had properties changed
+		IVsPersistSolutionProps, // Allows us to read props added to the solution to determine if the solution is a unreal solution
+		IDisposable
+	{
+		/** Constants */
+
+		private const string VersionString = "v1.60";
+		private const string UnrealSolutionFileNamePrefix = "UE";
+		private const string ExtensionName = "UnrealVS";
+		private const string CommandLineOptionKey = ExtensionName + "CommandLineMRU";
+		private const string BatchBuildSetsOptionKey = ExtensionName + "BatchBuildSetsV002";
+		private readonly TimeSpan TickPeriod = TimeSpan.FromSeconds(1.0);
+
+		/** Events */
+
+		/// Called when a new startup project is set in Visual Studio
+		public delegate void OnStartupProjectChangedDelegate(Project NewStartupProject);
+		public event OnStartupProjectChangedDelegate OnStartupProjectChanged;
+
+		/// Called when a project is opened or created in Visual Studio
+		public delegate void OnProjectOpenedDelegate(Project OpenedProject);
+		public event OnProjectOpenedDelegate OnProjectOpened;
+
+		/// Called right before a project is closed
+		public delegate void OnProjectClosedDelegate(Project ClosedProject);
+		public event OnProjectClosedDelegate OnProjectClosed;
+
+		/// Called when the startup project is edited in Visual Studio
+		public delegate void OnStartupProjectPropertyChangedDelegate(UInt32 itemid, Int32 propid, UInt32 flags);
+		public event OnStartupProjectPropertyChangedDelegate OnStartupProjectPropertyChanged;
+
+		/// Called right after a solution is opened
+		public delegate void OnSolutionOpenedDelegate();
+		public event OnSolutionOpenedDelegate OnSolutionOpened;
+
+		/// Called right before/after a solution is closed
+		public delegate void OnSolutionClosedDelegate();
+		public event OnSolutionClosedDelegate OnSolutionClosing;
+		public event OnSolutionClosedDelegate OnSolutionClosed;
+
+		/// Called when the active project config changes for any project
+		public delegate void OnStartupProjectConfigChangedDelegate(Project Project);
+		public event OnStartupProjectConfigChangedDelegate OnStartupProjectConfigChanged;
+
+		/// Called when a build/update action begins
+		public delegate void OnBuildBeginDelegate(out int Cancel);
+		public event OnBuildBeginDelegate OnBuildBegin;
+
+		/// Called when a build/update action completes
+		public delegate void OnBuildDoneDelegate(bool bSucceeded, bool bModified, bool bWasCancelled);
+		public event OnBuildDoneDelegate OnBuildDone;
+
+		/// Called when the UIContext changes
+		public delegate void OnUIContextChangedDelegate(uint CmdUICookie, bool bActive);
+		public event OnUIContextChangedDelegate OnUIContextChanged;
+
+		/** Public Fields & Properties */
+
+		/// Returns singleton instance of UnrealVSPackage
+		public static UnrealVSPackage Instance
+		{
+			get { return PrivateInstance; }
+		}
+
+		/// Visual Studio menu command service
+		public IMenuCommandService MenuCommandService { get; private set; }
+
+		/// Visual Studio solution build manager interface.  This is used to change the active startup
+		/// Project, among other things.  We expose public access to the solution build manager through
+		/// our singleton instance
+		public IVsSolutionBuildManager2 SolutionBuildManager { get; private set; }
+
+
+		/// Visual Studio solution "manager" interface.  We register with this to receive events
+		/// about projects being added and such.  This needs to be cleaned up at shutdown.
+		public IVsSolution2 SolutionManager { get; private set; }
+
+		/// Our startup project selector component
+		public StartupProjectSelector StartupProjectSelector { get; private set; }
+
+		/// Our quick build component
+		public QuickBuild QuickBuilder { get; private set; }
+
+		/// Visual Studio shell selection manager interface.  Used to receive notifications about
+		/// startup projects changes, among other things.
+		public IVsMonitorSelection SelectionManager { get; private set; }
+
+		/// Variable keeps track of whether a supported Unreal solution is loaded
+		public bool IsUESolutionLoaded
+		{
+			get { return _IsUESolutionLoaded.GetValueOrDefault(false); }
+			private set { _IsUESolutionLoaded = value; }
+		}
+
+		/// Variable keeps track of the loaded solution
+		private string _SolutionFilepath;
+		public string SolutionFilepath { get { return _SolutionFilepath; } }
+
+		public UnrealSolutionProperties UnrealSolutionProperties { get; private set; }
+
+		/** Methods */
+
+		/// <summary>
+		/// Package constructor.  The package is being created but Visual Studio isn't fully initialized yet, so
+		/// it's NOT SAFE to call into Visual Studio services from here.  Do that in Initialize() instead.
+		/// </summary>
+		public UnrealVSPackage()
+		{
+			// Register this key string so the package can save command line data to solution options files.
+			// See OnLoadOptions() & OnSaveOptions()
+			AddOptionKey(CommandLineOptionKey);
+			AddOptionKey(BatchBuildSetsOptionKey);
+
+			// Setup singleton instance
+			PrivateInstance = this;
+
+			Logging.Initialize(ExtensionName, VersionString);
+			Logging.WriteLine("Loading UnrealVS extension package...");
+		}
+
+
+		/// <summary>
+		/// Initializes the package right after it's been "sited" into the fully-initialized Visual Studio IDE.
+		/// </summary>
+		protected override async System.Threading.Tasks.Task InitializeAsync(CancellationToken cancellationToken, IProgress<ServiceProgressData> progress)
+		{
+			Logging.WriteLine("Initializing UnrealVS extension...");
+
+			await JoinableTaskFactory.SwitchToMainThreadAsync(cancellationToken);
+
+			// Grab the MenuCommandService
+			MenuCommandService = await GetServiceAsync(typeof(IMenuCommandService)) as OleMenuCommandService;
+
+			// Get access to Visual Studio's DTE object.  This object has various hooks into the Visual Studio
+			// shell that are useful for writing extensions.
+			DTE = await GetServiceAsync(typeof(DTE)) as DTE;
+			Logging.WriteLine("DTE version " + DTE.Version);
+
+			// Get selection manager and register to receive events
+			SelectionManager =
+				await GetServiceAsync(typeof(SVsShellMonitorSelection)) as IVsMonitorSelection;
+			SelectionManager.AdviseSelectionEvents(this, out SelectionEventsHandle);
+
+			// Get solution and register to receive events
+			SolutionManager = await GetServiceAsync(typeof(SVsSolution)) as IVsSolution2;
+			UpdateUnrealLoadedStatus();
+			SolutionManager.AdviseSolutionEvents(this, out SolutionEventsHandle);
+
+			// Grab the solution build manager.  We need this in order to change certain things about the Visual
+			// Studio environment, like what the active startup project is
+			// Get solution build manager
+			SolutionBuildManager =
+				await GetServiceAsync(typeof(SVsSolutionBuildManager)) as IVsSolutionBuildManager2;
+			SolutionBuildManager.AdviseUpdateSolutionEvents(this, out UpdateSolutionEventsHandle);
+
+			// Create our command-line editor
+			CommandLineEditor.Initialize();
+
+			// Create our startup project selector
+			StartupProjectSelector = new StartupProjectSelector();
+
+			// Create 'BuildStartupProject' instance
+			BuildStartupProject = new BuildStartupProject();
+
+			// Create 'CompileSingleFile' instance
+			CompileSingleFile = new CompileSingleFile();
+
+			// Create 'GenerateProjectFiles' tools
+			GenerateProjectFiles = new GenerateProjectFiles();
+
+			// Create Batch Builder tools
+			BatchBuilder.Initialize();
+
+			// Create the project menu quick builder
+			QuickBuilder = new QuickBuild();
+
+			// Call parent implementation
+			base.Initialize();
+
+			if (DTE.Solution.IsOpen)
+			{
+				StartTicker();
+			}
+		}
+
+		private void StartTicker()
+		{
+			// Create a "ticker" on a background thread that ticks the package on the UI thread
+			Interlocked.Exchange(ref bCancelTicker, 0);
+			Ticker = new Thread(TickAsyncMain)
+			{
+				Priority = ThreadPriority.Lowest
+			};
+			Ticker.Start();
+		}
+
+		private void StopTicker()
+		{
+			if (bCancelTicker == 0)
+			{
+				Interlocked.Exchange(ref bCancelTicker, 1);
+			}
+		}
+
+		/// <summary>
+		/// Tick loop on worker thread
+		/// </summary>
+		private void TickAsyncMain()
+		{
+			try
+			{
+				while (true)
+				{
+					if (bCancelTicker != 0) return;
+					ThreadHelper.Generic.BeginInvoke(Tick);
+
+					if (bCancelTicker != 0) return;
+					Thread.Sleep(TickPeriod);
+				}
+			}
+			catch (ThreadAbortException)
+			{
+			}
+		}
+
+		/// <summary>
+		/// Tick function on main UI thread
+		/// </summary>
+		private void Tick()
+		{
+			BatchBuilder.Tick();
+		}
+
+		/// <summary>
+		/// Implementation from IDisposable
+		/// </summary>
+		public void Dispose()
+		{
+			Dispose(true);
+		}
+
+		/// IDispose pattern lets us clean up our stuff!
+		protected override void Dispose(bool disposing)
+		{
+			if (Ticker != null && Ticker.IsAlive)
+			{
+				Thread.Sleep(TickPeriod + TickPeriod);
+				if (Ticker.IsAlive)
+				{
+					Logging.WriteLine("WARNING: Force aborting Ticker thread");
+					Ticker.Abort();
+				}
+			}
+
+			base.Dispose(disposing);
+
+			// Clean up singleton instance
+			PrivateInstance = null;
+
+			StartupProjectSelector = null;
+			BatchBuilder = null;
+			QuickBuilder = null;
+
+			CompileSingleFile?.Dispose();
+			CompileSingleFile = null;
+
+			CommandLineEditor?.Dispose();
+			CommandLineEditor = null;
+
+			// No longer want solution events
+			if (SolutionEventsHandle != 0)
+			{
+				SolutionManager.UnadviseSolutionEvents(SolutionEventsHandle);
+				SolutionEventsHandle = 0;
+			}
+			SolutionManager = null;
+
+			// No longer want selection events
+			if (SelectionEventsHandle != 0)
+			{
+				SelectionManager.UnadviseSelectionEvents(SelectionEventsHandle);
+				SelectionEventsHandle = 0;
+			}
+			SelectionManager = null;
+
+			// No longer want update solution events
+			if (UpdateSolutionEventsHandle != 0)
+			{
+				SolutionBuildManager.UnadviseUpdateSolutionEvents(UpdateSolutionEventsHandle);
+				UpdateSolutionEventsHandle = 0;
+			}
+			SolutionBuildManager = null;
+
+			Logging.WriteLine("Closing UnrealVS extension");
+			Logging.Close();
+		}
+
+
+		/// Visual Studio shell DTE interface
+		public DTE DTE
+		{
+			get;
+			private set;
+		}
+
+		/// Visual Studio shell DTE2 interface
+		private DTE2 _DTE2;
+		public DTE2 DTE2
+		{
+			get
+			{
+				if (_DTE2 == null)
+				{
+					// Get the interface when first used.
+					// This method fails during the Initialize() of the 
+					// package due to the strange method it requires.
+					_DTE2 = GetDTE2ForCurrentInstance(DTE);
+				}
+				return _DTE2;
+			}
+		}
+
+		/// The package's options page
+		public UnrealVsOptions OptionsPage
+		{
+			get { return (UnrealVsOptions)GetDialogPage(typeof(UnrealVsOptions)); }
+		}
+
+		/// <summary>
+		/// Launches a program
+		/// </summary>
+		/// <param name="ProgramFile">Path to the program to run</param>
+		/// <param name="Arguments">Command-line arguments</param>
+		/// <param name="OnExit">Optional callback whent he program exits</param>
+		/// <param name="OutputHandler">If supplied, std-out and std-error will be redirected to this function and no shell window will be created</param>
+		/// <returns>The newly-created process, or null if it wasn't able to start.  Exceptions are swallowed, but a debug output string is emitted on errors</returns>
+		public static System.Diagnostics.Process LaunchProgram(string ProgramFile, string Arguments, EventHandler OnExit = null, DataReceivedEventHandler OutputHandler = null, bool bWaitForCompletion = false)
+		{
+			// Create the action's process.
+			ProcessStartInfo ActionStartInfo = new ProcessStartInfo
+			{
+				FileName = ProgramFile,
+				Arguments = Arguments
+			};
+
+			if (OutputHandler != null)
+			{
+				ActionStartInfo.RedirectStandardInput = true;
+				ActionStartInfo.RedirectStandardOutput = true;
+				ActionStartInfo.RedirectStandardError = true;
+
+				// True to use a DOS box to run the program in, otherwise false to run directly.  In order to redirect
+				// output, UseShellExecute must be disabled
+				ActionStartInfo.UseShellExecute = false;
+
+				// Don't show the DOS box, since we're redirecting everything
+				ActionStartInfo.CreateNoWindow = true;
+			}
+
+
+			Logging.WriteLine(String.Format("Executing: {0} {1}", ActionStartInfo.FileName, ActionStartInfo.Arguments));
+
+			System.Diagnostics.Process ActionProcess;
+
+			try
+			{
+				ActionProcess = new System.Diagnostics.Process
+				{
+					StartInfo = ActionStartInfo
+				};
+
+				if (OnExit != null)
+				{
+					ActionProcess.EnableRaisingEvents = true;
+					ActionProcess.Exited += OnExit;
+				}
+
+				if (ActionStartInfo.RedirectStandardOutput)
+				{
+					ActionProcess.EnableRaisingEvents = true;
+					ActionProcess.OutputDataReceived += OutputHandler;
+					ActionProcess.ErrorDataReceived += OutputHandler;
+				}
+
+				// Launch the program
+				ActionProcess.Start();
+
+				if (ActionStartInfo.RedirectStandardOutput)
+				{
+					ActionProcess.BeginOutputReadLine();
+					ActionProcess.BeginErrorReadLine();
+				}
+
+				if (bWaitForCompletion)
+				{
+					while ((ActionProcess != null) && (!ActionProcess.HasExited))
+					{
+						ActionProcess.WaitForExit(50);
+					}
+				}
+			}
+			catch (Exception Ex)
+			{
+				// Couldn't launch program
+				Logging.WriteLine("Couldn't launch program: " + ActionStartInfo.FileName);
+				Logging.WriteLine("Exception: " + Ex.Message);
+				ActionProcess = null;
+			}
+
+			return ActionProcess;
+		}
+
+
+		/// <summary>
+		/// Gets a Visual Studio pane to output text to, or creates one if not visible.  Does not bring the pane to front (you can call Activate() to do that.)
+		/// </summary>
+		/// <returns>The pane to output to, or null on error</returns>
+		public IVsOutputWindowPane GetOutputPane()
+		{
+			return GetOutputPane(VSConstants.OutputWindowPaneGuid.BuildOutputPane_guid, "Build");
+		}
+
+		public IEnumerable<string> GetLoadedProjectPaths()
+		{
+			return LoadedProjectPaths;
+		}
+
+		/// <summary>
+		/// Overrides Package.OnLoadOptions()
+		/// Invoked by the package class when there are options to be read out of the solution file.
+		/// </summary>
+		/// <param name="key">The name of the option key to load.</param>
+		/// <param name="stream">The stream to load the option data from.</param>
+		protected override void OnLoadOptions(string key, Stream stream)
+		{
+			Logging.WriteLine("Loading Options for key: " + key);
+			try
+			{
+				if (0 == string.Compare(key, CommandLineOptionKey))
+				{
+					Logging.WriteLine("Restoring CommandLineEditor options");
+					CommandLineEditor.LoadOptions(stream);
+				}
+				else if (0 == string.Compare(key, BatchBuildSetsOptionKey))
+				{
+					Logging.WriteLine("Restoring BatchBuilder options");
+					BatchBuilder.LoadOptions(stream);
+				}
+			}
+			catch (Exception Ex)
+			{
+				// Couldn't load options
+				Exception AppEx = new ApplicationException("OnLoadOptions() failed with key " + key, Ex);
+				Logging.WriteLine(AppEx.ToString());
+				throw AppEx;
+			}
+		}
+
+
+		/// <summary>
+		/// Overrides Package.OnSaveOptions()
+		/// Invoked by the Package class when there are options to be saved to the solution file.
+		/// </summary>
+		/// <param name="key">The name of the option key to save.</param>
+		/// <param name="stream">The stream to save the option data to.</param>
+		protected override void OnSaveOptions(string key, Stream stream)
+		{
+			try
+			{
+				if (0 == string.Compare(key, CommandLineOptionKey))
+				{
+					Logging.WriteLine("Saving CommandLineEditor options");
+					CommandLineEditor.SaveOptions(stream);
+				}
+				else if (0 == string.Compare(key, BatchBuildSetsOptionKey))
+				{
+					Logging.WriteLine("Saving BatchBuilder options");
+					BatchBuilder.SaveOptions(stream);
+				}
+			}
+			catch (Exception Ex)
+			{
+				// Couldn't save options
+				Exception AppEx = new ApplicationException("OnSaveOptions() failed with key " + key, Ex);
+				Logging.WriteLine(AppEx.ToString());
+				throw AppEx;
+			}
+		}
+
+
+		///
+		/// IVsSolutionEvents implementation
+		///
+
+		int IVsSolutionEvents.OnAfterCloseSolution(object pUnkReserved)
+		{
+			UpdateUnrealLoadedStatus();
+
+			OnSolutionClosed?.Invoke();
+
+			return VSConstants.S_OK;
+		}
+
+		int IVsSolutionEvents.OnAfterLoadProject(IVsHierarchy pStubHierarchy, IVsHierarchy pRealHierarchy)
+		{
+			return VSConstants.S_OK;
+		}
+
+		int IVsSolutionEvents.OnAfterOpenProject(IVsHierarchy pHierarchy, int fAdded)
+		{
+			// This function is called after a Visual Studio project is opened (or a new project is created.)
+
+			// Get the actual Project object from the IVsHierarchy object that was supplied
+			var OpenedProject = Utils.HierarchyObjectToProject(pHierarchy);
+			Utils.OnProjectListChanged();
+			if (OpenedProject != null && OnProjectOpened != null)
+			{
+				LoadedProjectPaths.Add(OpenedProject.FullName);
+				OnProjectOpened(OpenedProject);
+			}
+
+			return VSConstants.S_OK;
+		}
+
+		int IVsSolutionEvents.OnAfterOpenSolution(object pUnkReserved, int fNewSolution)
+		{
+			UpdateUnrealLoadedStatus();
+
+			StartTicker();
+
+			OnSolutionOpened?.Invoke();
+
+			return VSConstants.S_OK;
+		}
+
+		int IVsSolutionEvents.OnBeforeCloseProject(IVsHierarchy pHierarchy, int fRemoved)
+		{
+			// This function is called after a Visual Studio project is closed
+
+			// Get the actual Project object from the IVsHierarchy object that was supplied
+			var ClosedProject = Utils.HierarchyObjectToProject(pHierarchy);
+			if (ClosedProject != null && OnProjectClosed != null)
+			{
+				LoadedProjectPaths.Remove(ClosedProject.FullName);
+				OnProjectClosed(ClosedProject);
+			}
+
+			return VSConstants.S_OK;
+		}
+
+		int IVsSolutionEvents.OnBeforeCloseSolution(object pUnkReserved)
+		{
+			StopTicker();
+
+			OnSolutionClosing?.Invoke();
+
+			return VSConstants.S_OK;
+		}
+
+		int IVsSolutionEvents.OnBeforeUnloadProject(IVsHierarchy pRealHierarchy, IVsHierarchy pStubHierarchy)
+		{
+			return VSConstants.S_OK;
+		}
+
+		int IVsSolutionEvents.OnQueryCloseProject(IVsHierarchy pHierarchy, int fRemoving, ref int pfCancel)
+		{
+			return VSConstants.S_OK;
+		}
+
+		int IVsSolutionEvents.OnQueryCloseSolution(object pUnkReserved, ref int pfCancel)
+		{
+			if (BatchBuilder.IsBusy)
+			{
+				pfCancel = 1;
+			}
+
+			return VSConstants.S_OK;
+		}
+
+		int IVsSolutionEvents.OnQueryUnloadProject(IVsHierarchy pRealHierarchy, ref int pfCancel)
+		{
+			return VSConstants.S_OK;
+		}
+
+
+		///
+		/// IVsSelectionEvents implementation
+		///
+
+		int IVsSelectionEvents.OnCmdUIContextChanged(uint dwCmdUICookie, int fActive)
+		{
+			OnUIContextChanged?.Invoke(dwCmdUICookie, fActive != 0);
+
+			return VSConstants.S_OK;
+		}
+
+		int IVsSelectionEvents.OnElementValueChanged(uint elementid, object varValueOld, object varValueNew)
+		{
+			// This function is called when selection changes in various Visual Studio tool windows
+			// and sub-systems.
+
+			// Handle startup project changes
+			if (elementid == (uint)VSConstants.VSSELELEMID.SEID_StartupProject)
+			{
+				// If we are registered to a project hierarchy for events, unregister
+				var OldStartupProjectHierarchy = (IVsHierarchy)varValueOld;
+				if (OldStartupProjectHierarchy != null && ProjectHierarchyEventsHandle != 0)
+				{
+					OldStartupProjectHierarchy.UnadviseHierarchyEvents(ProjectHierarchyEventsHandle);
+					ProjectHierarchyEventsHandle = 0;
+				}
+
+				Project NewStartupProject = null;
+
+				// Incoming hierarchy object could be null (if no startup project is set yet, or during shutdown.)
+				var NewStartupProjectHierarchy = (IVsHierarchy)varValueNew;
+				if (NewStartupProjectHierarchy != null)
+				{
+					// Get the actual Project object from the IVsHierarchy object that was supplied
+					NewStartupProject = Utils.HierarchyObjectToProject(NewStartupProjectHierarchy);
+
+					if (NewStartupProject != null)
+					{
+						// Register for events from the project
+						NewStartupProjectHierarchy.AdviseHierarchyEvents(this, out ProjectHierarchyEventsHandle);
+					}
+				}
+
+				if (NewStartupProject != null)
+				{
+					OnStartupProjectChanged(NewStartupProject);
+				}
+			}
+
+			return VSConstants.S_OK;
+		}
+
+		int IVsSelectionEvents.OnSelectionChanged(IVsHierarchy pHierOld, uint itemidOld, IVsMultiItemSelect pMISOld,
+												  ISelectionContainer pSCOld, IVsHierarchy pHierNew, uint itemidNew,
+												  IVsMultiItemSelect pMISNew, ISelectionContainer pSCNew)
+		{
+			return VSConstants.S_OK;
+		}
+
+
+		// IVsHierarchyEvents Interface
+
+		Int32 IVsHierarchyEvents.OnItemAdded(UInt32 itemidParent, UInt32 itemidSiblingPrev, UInt32 itemidAdded)
+		{
+			return VSConstants.S_OK;
+		}
+
+		Int32 IVsHierarchyEvents.OnPropertyChanged(UInt32 itemid, Int32 propid, UInt32 flags)
+		{
+			OnStartupProjectPropertyChanged?.Invoke(itemid, propid, flags);
+			return VSConstants.S_OK;
+		}
+
+		Int32 IVsHierarchyEvents.OnItemsAppended(UInt32 itemidParent)
+		{
+			return VSConstants.S_OK;
+		}
+
+		Int32 IVsHierarchyEvents.OnItemDeleted(UInt32 itemid)
+		{
+			return VSConstants.S_OK;
+		}
+
+		Int32 IVsHierarchyEvents.OnInvalidateItems(UInt32 itemidParent)
+		{
+			return VSConstants.S_OK;
+		}
+
+		Int32 IVsHierarchyEvents.OnInvalidateIcon(IntPtr hicon)
+		{
+			return VSConstants.S_OK;
+		}
+
+
+		#region IVsPersistSolutionProps
+		int IVsPersistSolutionProps.SaveUserOptions(IVsSolutionPersistence pPersistence)
+		{
+			return VSConstants.S_OK;
+		}
+
+		int IVsPersistSolutionProps.LoadUserOptions(IVsSolutionPersistence pPersistence, uint grfLoadOpts)
+		{
+			return VSConstants.S_OK;
+		}
+
+		int IVsPersistSolutionProps.WriteUserOptions(IStream pOptionsStream, string pszKey)
+		{
+			return VSConstants.S_OK;
+		}
+
+		int IVsPersistSolutionProps.ReadUserOptions(IStream pOptionsStream, string pszKey)
+		{
+			return VSConstants.S_OK;
+		}
+
+		int IVsPersistSolutionProps.QuerySaveSolutionProps(IVsHierarchy pHierarchy, VSQUERYSAVESLNPROPS[] pqsspSave)
+		{
+			return VSConstants.S_OK;
+		}
+
+		int IVsPersistSolutionProps.SaveSolutionProps(IVsHierarchy pHierarchy, IVsSolutionPersistence pPersistence)
+		{
+			return VSConstants.S_OK;
+		}
+
+		int IVsPersistSolutionProps.WriteSolutionProps(IVsHierarchy pHierarchy, string pszKey, IPropertyBag pPropBag)
+		{
+			return VSConstants.S_OK;
+		}
+
+		int IVsPersistSolutionProps.ReadSolutionProps(IVsHierarchy pHierarchy, string pszProjectName, string pszProjectMk, string pszKey, int fPreLoad, IPropertyBag pPropBag)
+		{
+			if (!string.Equals(pszKey, GuidList.UnrealVSPackageString, StringComparison.InvariantCultureIgnoreCase))
+				return VSConstants.S_OK;
+
+			pPropBag.Read("AvailablePlatforms", out object availablePlatformsObject, null, (uint)VarEnum.VT_BSTR, pPropBag);
+
+			string[] availablePlatforms = null;
+			if (availablePlatformsObject != null)
+			{
+				string availablePlatformsString = availablePlatformsObject as string;
+				availablePlatforms = availablePlatformsString.Split(';');
+			}
+
+			UnrealSolutionProperties = new UnrealSolutionProperties(availablePlatforms);
+
+			// a UnrealVS section was found so we consider this solution a unreal solution
+			IsUESolutionLoaded = true;
+
+			return VSConstants.S_OK;
+		}
+
+		int IVsPersistSolutionProps.OnProjectLoadFailure(IVsHierarchy pStubHierarchy, string pszProjectName, string pszProjectMk, string pszKey)
+		{
+			return VSConstants.S_OK;
+		}
+
+		#endregion
+
+		// IVsUpdateSolutionEvents Interface
+
+		public int UpdateSolution_Begin(ref int pfCancelUpdate)
+		{
+			OnBuildBegin?.Invoke(out pfCancelUpdate);
+			return VSConstants.S_OK;
+		}
+
+		public int UpdateSolution_Done(int fSucceeded, int fModified, int fCancelCommand)
+		{
+			OnBuildDone?.Invoke(fSucceeded != 0, fModified != 0, fCancelCommand != 0);
+			return VSConstants.S_OK;
+		}
+
+		public int UpdateSolution_StartUpdate(ref int pfCancelUpdate)
+		{
+			return VSConstants.S_OK;
+		}
+
+		public int UpdateSolution_Cancel()
+		{
+			return VSConstants.S_OK;
+		}
+
+		public int OnActiveProjectCfgChange(IVsHierarchy pIVsHierarchy)
+		{
+			// This function is called after a Visual Studio project has its active config changed
+
+			// Check whether the project is the current startup project
+			SolutionBuildManager.get_StartupProject(out IVsHierarchy StartupProjectHierarchy);
+			if (StartupProjectHierarchy != null && StartupProjectHierarchy == pIVsHierarchy)
+			{
+				// Get the actual Project object from the IVsHierarchy object that was supplied
+				var Project = Utils.HierarchyObjectToProject(pIVsHierarchy);
+				if (Project != null && OnStartupProjectConfigChanged != null)
+				{
+					OnStartupProjectConfigChanged(Project);
+				}
+			}
+			return VSConstants.S_OK;
+		}
+
+		private void UpdateUnrealLoadedStatus()
+		{
+			if (!DTE.Solution.IsOpen)
+			{
+				IsUESolutionLoaded = false;
+				return;
+			}
+
+			SolutionManager.GetSolutionInfo(out string SolutionDirectory, out _SolutionFilepath, out string UserOptsFile);
+
+			// if ReadProps found a valud UE solution we do not need to check these legacy definitions of a solution
+			if (_IsUESolutionLoaded.HasValue)
+				return;
+
+			// Legacy paths for determing if a solution is a unreal solution by checking for tags & solution name.
+			string[] SolutionLines = Array.Empty<string>();
+			try
+			{
+				SolutionLines = File.ReadAllLines(_SolutionFilepath);
+			}
+			catch
+			{
+			}
+
+			const string UBTTag = "# UnrealEngineGeneratedSolutionVersion=";
+			var UBTLine = SolutionLines.FirstOrDefault(TextLine => TextLine.Trim().StartsWith(UBTTag));
+			if (UBTLine != null)
+			{
+				_UBTVersion = UBTLine.Trim().Substring(UBTTag.Length);
+				IsUESolutionLoaded = true;
+			}
+			else
+			{
+				_UBTVersion = string.Empty;
+				IsUESolutionLoaded =
+					(
+						_SolutionFilepath != null &&
+						Path.GetFileName(_SolutionFilepath).StartsWith(UnrealSolutionFileNamePrefix, StringComparison.OrdinalIgnoreCase)
+					);
+			}
+		}
+
+		/** Private Fields & Properties */
+
+		private static UnrealVSPackage PrivateInstance = null;
+
+		/// Handle that we used at shutdown to unregister for selection manager events
+		private UInt32 SelectionEventsHandle;
+
+		/// Handle that we use at shutdown to unregister for events about solution activity
+		private UInt32 SolutionEventsHandle;
+
+		/// Handle that we use at shutdown to unregister for events about solution build activity
+		private UInt32 UpdateSolutionEventsHandle;
+
+		/// Handle that we use to unregister for events about startup project hierarchy activity
+		UInt32 ProjectHierarchyEventsHandle;
+
+		/// Our command-line editing component
+		private CommandLineEditor CommandLineEditor = new CommandLineEditor();
+
+		/// BuildStartupProject feature
+		private BuildStartupProject BuildStartupProject;
+
+		/// CompileSingleFile feature
+		private CompileSingleFile CompileSingleFile;
+
+		/// Project file generator button
+		private GenerateProjectFiles GenerateProjectFiles;
+
+		/// Batch Builder button/command handler
+		private BatchBuilder BatchBuilder = new BatchBuilder();
+
+		/// Ticker thread
+		private Thread Ticker;
+
+		/// Ticker thread cancel flag
+		private int bCancelTicker = 0;
+
+		private string _UBTVersion = string.Empty;
+
+		private readonly List<string> LoadedProjectPaths = new List<string>();
+		private bool? _IsUESolutionLoaded;
+
+		/// Obtains the DTE2 interface for this instance of VS from the RunningObjectTable
+		private static DTE2 GetDTE2ForCurrentInstance(DTE DTE)
+		{
+			// Find the ROT entry for visual studio running under current process.
+			int HResult = NativeMethods.GetRunningObjectTable(0, out IRunningObjectTable Rot);
+			if (HResult == 0)
+			{
+				Rot.EnumRunning(out IEnumMoniker EnumMoniker);
+				EnumMoniker.Reset();
+				IMoniker[] Moniker = new IMoniker[1];
+				while (EnumMoniker.Next(1, Moniker, out _) == 0)
+				{
+					Rot.GetObject(Moniker[0], out object ComObject);
+
+					if (ComObject is DTE2 CandidateDTE2)
+					{
+						if (CandidateDTE2.DTE == DTE)
+						{
+							return CandidateDTE2;
+						}
+					}
+				}
+			}
+			return null;
+		}
+
+	}
+
+	internal static class NativeMethods
+	{
+		/// ROT function in ole32.dll needed by GetDTE2ForCurrentInstance()
+		[DllImport("ole32.dll")]
+		internal static extern int GetRunningObjectTable(int reserved, out IRunningObjectTable prot);
+	}
+}
diff --git a/Engine/Source/Programs/UnrealVS/UnrealVS.2019/Utils.cs b/Engine/Source/Programs/UnrealVS/UnrealVS.2019/Utils.cs
new file mode 100644
--- /dev/null
+++ b/Engine/Source/Programs/UnrealVS/UnrealVS.2019/Utils.cs
@@ -0,0 +1,691 @@
+// Copyright Epic Games, Inc. All Rights Reserved.
+
+using EnvDTE;
+using EnvDTE80;
+using Microsoft.VisualStudio;
+using Microsoft.VisualStudio.Shell.Interop;
+using System;
+using System.Collections.Generic;
+using System.IO;
+using System.Linq;
+using System.Xml;
+
+namespace UnrealVS
+{
+	public class Utils
+	{
+		public const string UProjectExtension = "uproject";
+
+		public class SafeProjectReference
+		{
+			public string FullName { get; set; }
+			public string Name { get; set; }
+
+			public Project GetProjectSlow()
+			{
+				Project[] Projects = GetAllProjectsFromDTE();
+				return Projects.FirstOrDefault(Proj => string.CompareOrdinal(Proj.FullName, FullName) == 0);
+			}
+		}
+
+		/// <summary>
+		/// Converts a Project to an IVsHierarchy
+		/// </summary>
+		/// <param name="Project">Project object</param>
+		/// <returns>IVsHierarchy for the specified project</returns>
+		public static IVsHierarchy ProjectToHierarchyObject(Project Project)
+		{
+			UnrealVSPackage.Instance.SolutionManager.GetProjectOfUniqueName(Project.FullName, out IVsHierarchy HierarchyObject);
+			return HierarchyObject;
+		}
+
+
+		/// <summary>
+		/// Converts an IVsHierarchy object to a Project
+		/// </summary>
+		/// <param name="HierarchyObject">IVsHierarchy object</param>
+		/// <returns>Visual Studio project object</returns>
+		public static Project HierarchyObjectToProject(IVsHierarchy HierarchyObject)
+		{
+			// Get the actual Project object from the IVsHierarchy object that was supplied
+			HierarchyObject.GetProperty(VSConstants.VSITEMID_ROOT, (int)__VSHPROPID.VSHPROPID_ExtObject, out object ProjectObject);
+			return (Project)ProjectObject;
+		}
+
+		/// <summary>
+		/// Converts an IVsHierarchy object to a config provider interface
+		/// </summary>
+		/// <param name="HierarchyObject">IVsHierarchy object</param>
+		/// <returns>Visual Studio project object</returns>
+		public static IVsCfgProvider2 HierarchyObjectToCfgProvider(IVsHierarchy HierarchyObject)
+		{
+			// Get the actual Project object from the IVsHierarchy object that was supplied
+			HierarchyObject.GetProperty(VSConstants.VSITEMID_ROOT, (int)__VSHPROPID.VSHPROPID_BrowseObject, out object BrowseObject);
+
+			IVsCfgProvider2 CfgProvider = null;
+			if (BrowseObject != null)
+			{
+				CfgProvider = GetCfgProviderFromObject(BrowseObject);
+			}
+
+			if (CfgProvider == null)
+			{
+				CfgProvider = GetCfgProviderFromObject(HierarchyObject);
+			}
+
+			return CfgProvider;
+		}
+
+		private static IVsCfgProvider2 GetCfgProviderFromObject(object SomeObject)
+		{
+			IVsCfgProvider2 CfgProvider2 = null;
+
+			if (SomeObject is IVsGetCfgProvider GetCfgProvider)
+			{
+				GetCfgProvider.GetCfgProvider(out IVsCfgProvider CfgProvider);
+				if (CfgProvider != null)
+				{
+					CfgProvider2 = CfgProvider as IVsCfgProvider2;
+				}
+			}
+
+			if (CfgProvider2 == null)
+			{
+				CfgProvider2 = SomeObject as IVsCfgProvider2;
+			}
+
+			return CfgProvider2;
+		}
+
+		/// <summary>
+		/// Locates a specific project property for the active configuration and returns it (or null if not found.)
+		/// </summary>
+		/// <param name="Project">Project to search for the property</param>
+		/// <param name="PropertyName">Name of the property</param>
+		/// <returns>Property object or null if not found</returns>
+		public static Property GetProjectProperty(Project Project, string PropertyName)
+		{
+			var Properties = Project.Properties;
+			if (Properties != null)
+			{
+				foreach (var RawProperty in Properties)
+				{
+					var Property = (Property)RawProperty;
+					if (Property.Name.Equals(PropertyName, StringComparison.InvariantCultureIgnoreCase))
+					{
+						return Property;
+					}
+				}
+			}
+
+			// Not found
+			return null;
+		}
+
+		/// <summary>
+		/// Locates a specific project property for the active configuration and attempts to set its value
+		/// </summary>
+		/// <param name="Property">The property object to set</param>
+		/// <param name="PropertyValue">Value to set for this property</param>
+		public static void SetPropertyValue(Property Property, object PropertyValue)
+		{
+			Property.Value = PropertyValue;
+
+			// @todo: Not sure if actually needed for command-line property (saved in .user files, not in project)
+			// Mark the project as modified
+			// @todo: Throws exception for C++ projects, doesn't mark as saved
+			//				Project.IsDirty = true;
+			//				Project.Saved = false;
+		}
+
+		/// <summary>
+		/// Helper class used by the GetUIxxx functions below.
+		/// Callers use this to easily traverse UIHierarchies.
+		/// </summary>
+		public class UITreeItem
+		{
+			public UIHierarchyItem Item { get; set; }
+			public UITreeItem[] Children { get; set; }
+			public string Name { get { return Item != null ? Item.Name : "None"; } }
+			public object Object { get { return Item?.Object; } }
+		}
+
+		/// <summary>
+		/// Converts a UIHierarchy into an easy to use tree of helper class UITreeItem.
+		/// </summary>
+		public static UITreeItem GetUIHierarchyTree(UIHierarchy Hierarchy)
+		{
+			return new UITreeItem
+			{
+				Item = null,
+				Children = (from UIHierarchyItem Child in Hierarchy.UIHierarchyItems select GetUIHierarchyTree(Child)).ToArray()
+			};
+		}
+
+		/// <summary>
+		/// Called by the public GetUIHierarchyTree() function above.
+		/// </summary>
+		private static UITreeItem GetUIHierarchyTree(UIHierarchyItem HierarchyItem)
+		{
+			return new UITreeItem
+			{
+				Item = HierarchyItem,
+				Children = (from UIHierarchyItem Child in HierarchyItem.UIHierarchyItems select GetUIHierarchyTree(Child)).ToArray()
+			};
+		}
+
+		/// <summary>
+		/// Helper function to easily extract a list of objects of type T from a UIHierarchy tree.
+		/// </summary>
+		/// <typeparam name="T">The type of object to find in the tree. Extracts everything that "Is a" T.</typeparam>
+		/// <param name="RootItem">The root of the UIHierarchy to search (converted to UITreeItem via GetUIHierarchyTree())</param>
+		/// <returns>An enumerable of objects of type T found beneath the root item.</returns>
+		public static IEnumerable<T> GetUITreeItemObjectsByType<T>(UITreeItem RootItem) where T : class
+		{
+			List<T> Results = new List<T>();
+
+			if (RootItem.Object is T Obj)
+			{
+				Results.Add(Obj);
+			}
+			foreach (var Child in RootItem.Children)
+			{
+				Results.AddRange(GetUITreeItemObjectsByType<T>(Child));
+			}
+
+			return Results;
+		}
+
+		public static IEnumerable<UIHierarchyItem> GetUITreeItemsByObjectType<T>(UITreeItem RootItem) where T : class
+		{
+			List<UIHierarchyItem> Results = new List<UIHierarchyItem>();
+
+			if (RootItem.Object is T)
+			{
+				Results.Add(RootItem.Item);
+			}
+			foreach (var Child in RootItem.Children)
+			{
+				Results.AddRange(GetUITreeItemsByObjectType<T>(Child));
+			}
+
+			return Results;
+		}
+
+		/// <summary>
+		/// Helper to check the file ext of a binary against known library file exts.
+		/// FileExt should include the dot e.g. ".dll"
+		/// </summary>
+		public static bool IsLibraryFileExtension(string FileExt)
+		{
+			if (FileExt.Equals(".dll", StringComparison.InvariantCultureIgnoreCase)) return true;
+			if (FileExt.Equals(".lib", StringComparison.InvariantCultureIgnoreCase)) return true;
+			if (FileExt.Equals(".ocx", StringComparison.InvariantCultureIgnoreCase)) return true;
+			if (FileExt.Equals(".a", StringComparison.InvariantCultureIgnoreCase)) return true;
+			if (FileExt.Equals(".so", StringComparison.InvariantCultureIgnoreCase)) return true;
+			if (FileExt.Equals(".dylib", StringComparison.InvariantCultureIgnoreCase)) return true;
+
+			return false;
+		}
+
+		/// <summary>
+		/// Helper to check the properties of a project and determine whether it can be built in VS.
+		/// </summary>
+		public static bool IsProjectBuildable(Project Project)
+		{
+			return Project.Kind == GuidList.VCSharpProjectKindGuidString || Project.Kind == GuidList.VCProjectKindGuidString;
+		}
+
+		/// Helper function to get the full list of all projects in the DTE Solution
+		/// Recurses into items because these are actually in a tree structure
+		public static Project[] GetAllProjectsFromDTE()
+		{
+			try
+			{
+				List<Project> Projects = new List<Project>();
+
+				foreach (Project Project in UnrealVSPackage.Instance.DTE.Solution.Projects)
+				{
+					Projects.Add(Project);
+
+					if (Project.ProjectItems != null)
+					{
+						foreach (ProjectItem Item in Project.ProjectItems)
+						{
+							GetSubProjectsOfProjectItem(Item, Projects);
+						}
+					}
+				}
+
+				return Projects.ToArray();
+			}
+			catch (Exception ex)
+			{
+				Exception AppEx = new ApplicationException("GetAllProjectsFromDTE() failed", ex);
+				Logging.WriteLine(AppEx.ToString());
+				throw AppEx;
+			}
+		}
+
+		public static void ExecuteProjectBuild(Project Project,
+												string SolutionConfig,
+												string SolutionPlatform,
+												BatchBuilderToolControl.BuildJob.BuildJobType BuildType,
+												Action ExecutingDelegate,
+												Action FailedToStartDelegate)
+		{
+			IVsHierarchy ProjHierarchy = Utils.ProjectToHierarchyObject(Project);
+
+			if (ProjHierarchy != null)
+			{
+				SolutionConfigurations SolutionConfigs =
+					UnrealVSPackage.Instance.DTE.Solution.SolutionBuild.SolutionConfigurations;
+
+				var MatchedSolutionConfig =
+					(from SolutionConfiguration2 Sc in SolutionConfigs select Sc).FirstOrDefault(
+						Sc =>
+						String.CompareOrdinal(Sc.Name, SolutionConfig) == 0 && String.CompareOrdinal(Sc.PlatformName, SolutionPlatform) == 0);
+
+				if (MatchedSolutionConfig != null)
+				{
+					SolutionContext ProjectSolutionCtxt = MatchedSolutionConfig.SolutionContexts.Item(Project.UniqueName);
+
+					if (ProjectSolutionCtxt != null)
+					{
+						IVsCfgProvider2 CfgProvider2 = Utils.HierarchyObjectToCfgProvider(ProjHierarchy);
+						if (CfgProvider2 != null)
+						{
+							CfgProvider2.GetCfgOfName(ProjectSolutionCtxt.ConfigurationName, ProjectSolutionCtxt.PlatformName, out IVsCfg Cfg);
+
+							if (Cfg != null)
+							{
+								ExecutingDelegate?.Invoke();
+
+								int JobResult = VSConstants.E_FAIL;
+
+								if (BuildType == BatchBuilderToolControl.BuildJob.BuildJobType.Build)
+								{
+									JobResult =
+										UnrealVSPackage.Instance.SolutionBuildManager.StartUpdateSpecificProjectConfigurations(
+											1,
+											new[] { ProjHierarchy },
+											new[] { Cfg },
+											null,
+											new uint[] { 0 },
+											null,
+											(uint)VSSOLNBUILDUPDATEFLAGS.SBF_OPERATION_BUILD,
+											0);
+								}
+								else if (BuildType == BatchBuilderToolControl.BuildJob.BuildJobType.Rebuild)
+								{
+									JobResult =
+										UnrealVSPackage.Instance.SolutionBuildManager.StartUpdateSpecificProjectConfigurations(
+											1,
+											new[] { ProjHierarchy },
+											new[] { Cfg },
+											new uint[] { 0 },
+											null,
+											null,
+											(uint)(VSSOLNBUILDUPDATEFLAGS.SBF_OPERATION_BUILD | VSSOLNBUILDUPDATEFLAGS.SBF_OPERATION_FORCE_UPDATE),
+											0);
+								}
+								else if (BuildType == BatchBuilderToolControl.BuildJob.BuildJobType.Clean)
+								{
+									JobResult =
+										UnrealVSPackage.Instance.SolutionBuildManager.StartUpdateSpecificProjectConfigurations(
+											1,
+											new[] { ProjHierarchy },
+											new[] { Cfg },
+											new uint[] { 0 },
+											null,
+											null,
+											(uint)VSSOLNBUILDUPDATEFLAGS.SBF_OPERATION_CLEAN,
+											0);
+								}
+
+								if (JobResult == VSConstants.S_OK)
+								{
+									// Job running - show output
+									PrepareOutputPane();
+								}
+								else
+								{
+									FailedToStartDelegate?.Invoke();
+								}
+							}
+						}
+					}
+				}
+			}
+		}
+
+
+		private static bool LoadConfigFromUBT(Project SelectedProject)
+		{
+			string ProjectPath = Path.GetDirectoryName(SelectedProject.FullName);
+			string ConfigFileName = Path.Combine(ProjectPath, "UnrealVS.xml");
+
+			// only try to load the xml configuration once
+			if (CachedUBTConfigFileName != ConfigFileName)
+			{
+				CachedUBTConfigFileName = ConfigFileName;
+				CachedUBTConfigXml = null;
+
+				try
+				{
+					XmlDocument ConfigXml = new XmlDocument();
+					ConfigXml.Load(ConfigFileName);
+
+					CachedUBTConfigXml = ConfigXml.SelectSingleNode("UnrealVS");
+
+				}
+				catch (Exception)
+				{
+				}
+			}
+
+			return (CachedUBTConfigXml != null);
+		}
+
+
+		public static List<string> GetExtraDebuggerCommandArguments(string PlatformName, Project SelectedProject)
+		{
+			List<string> Result = new List<string>();
+
+			if (LoadConfigFromUBT(SelectedProject))
+			{
+				XmlNode PlatformNode = CachedUBTConfigXml.SelectSingleNode(PlatformName);
+				if (PlatformNode != null)
+				{
+					foreach (XmlNode ChildNode in PlatformNode.ChildNodes)
+					{
+						if (string.Equals(ChildNode.Name, "DebuggerName", StringComparison.CurrentCultureIgnoreCase))
+						{
+							Result.Add(ChildNode.InnerText);
+						}
+					}
+				}
+			}
+
+			return Result;
+		}
+
+		public static bool IsGameProject(Project Project)
+		{
+			return GetUProjects().ContainsKey(Project.Name);
+		}
+
+		/// <summary>
+		/// Does the config build something that takes a .uproject on the command line?
+		/// </summary>
+		public static bool HasUProjectCommandLineArg(string Config)
+		{
+			return Config.EndsWith("Editor", StringComparison.InvariantCultureIgnoreCase);
+		}
+
+		public static string GetUProjectFileName(Project Project)
+		{
+			return Project.Name + "." + UProjectExtension;
+		}
+
+		public static string GetAutoUProjectCommandLinePrefix(Project Project)
+		{
+			var UProjectFileName = GetUProjectFileName(Project);
+			var AllUProjects = GetUProjects();
+
+			if (!AllUProjects.TryGetValue(Project.Name, out string UProjectPath))
+			{
+				// Search the project folder
+				var ProjectFolder = Path.GetDirectoryName(Project.FullName);
+				var UProjUnderProject = Directory.GetFiles(ProjectFolder, UProjectFileName, SearchOption.TopDirectoryOnly);
+				if (UProjUnderProject.Length == 1)
+				{
+					UProjectPath = UProjUnderProject[0];
+				}
+			}
+
+			return '\"' + UProjectPath + '\"';
+		}
+
+		public static void AddProjects(DirectoryInfo ProjectDir, List<FileInfo> Files)
+		{
+			Files.AddRange(ProjectDir.EnumerateFiles("*.uproject"));
+		}
+
+		/// <summary>
+		/// Enumerate projects under the given directory
+		/// </summary>
+		/// <param name="SolutionDir">Base directory to enumerate</param>
+		/// <returns>List of project files</returns>
+		static List<FileInfo> EnumerateProjects(DirectoryInfo SolutionDir)
+		{
+			// Enumerate all the projects in the same directory as the solution. If there's one here, we don't need to consider any other.
+			List<FileInfo> ProjectFiles = new List<FileInfo>(SolutionDir.EnumerateFiles("*.uproject"));
+			if (ProjectFiles.Count == 0)
+			{
+				// Build a list of all the parent directories for projects. This includes the UE root, plus any directories referenced via .uprojectdirs files.
+				List<DirectoryInfo> ParentProjectDirs = new List<DirectoryInfo>
+				{
+					SolutionDir
+				};
+
+				// Read all the .uprojectdirs files
+				foreach (FileInfo ProjectDirsFile in SolutionDir.EnumerateFiles("*.uprojectdirs"))
+				{
+					foreach (string Line in File.ReadAllLines(ProjectDirsFile.FullName))
+					{
+						string TrimLine = Line.Trim().Replace(Path.AltDirectorySeparatorChar, Path.DirectorySeparatorChar).Trim(Path.DirectorySeparatorChar);
+						if (TrimLine.Length > 0 && !TrimLine.StartsWith(";"))
+						{
+							try
+							{
+								ParentProjectDirs.Add(new DirectoryInfo(Path.Combine(SolutionDir.FullName, TrimLine)));
+							}
+							catch (Exception Ex)
+							{
+								Logging.WriteLine(String.Format("EnumerateProjects: Exception trying to resolve project directory '{0}': {1}", TrimLine, Ex.Message));
+							}
+						}
+					}
+				}
+
+				// Add projects in any subfolders of the parent directories
+				HashSet<string> CheckedParentDirs = new HashSet<string>(StringComparer.InvariantCultureIgnoreCase);
+				foreach (DirectoryInfo ParentProjectDir in ParentProjectDirs)
+				{
+					if (CheckedParentDirs.Add(ParentProjectDir.FullName) && ParentProjectDir.Exists)
+					{
+						foreach (DirectoryInfo ProjectDir in ParentProjectDir.EnumerateDirectories())
+						{
+							ProjectFiles.AddRange(ProjectDir.EnumerateFiles("*.uproject"));
+						}
+					}
+				}
+			}
+			return ProjectFiles;
+		}
+
+		/// <summary>
+		/// Returns all the .uprojects found under the solution root folder.
+		/// </summary>
+		public static IDictionary<string, string> GetUProjects()
+		{
+			var Folder = GetSolutionFolder();
+			if (string.IsNullOrEmpty(Folder))
+			{
+				return new Dictionary<string, string>();
+			}
+
+			if (Folder != CachedUProjectRootFolder)
+			{
+				Logging.WriteLine("GetUProjects: recaching uproject paths...");
+				DateTime Start = DateTime.Now;
+
+				CachedUProjectRootFolder = Folder;
+				CachedUProjectPaths = EnumerateProjects(new DirectoryInfo(Folder)).Select(x => x.FullName);
+				CachedUProjects = null;
+
+				TimeSpan TimeTaken = DateTime.Now - Start;
+				Logging.WriteLine(string.Format("GetUProjects: EnumerateProjects took {0} sec", TimeTaken.TotalSeconds));
+
+				foreach (string CachedUProjectPath in CachedUProjectPaths)
+				{
+					Logging.WriteLine(String.Format("GetUProjects: found {0}", CachedUProjectPath));
+				}
+
+				Logging.WriteLine("    DONE");
+			}
+
+			if (CachedUProjects == null)
+			{
+				Logging.WriteLine("GetUProjects: recaching uproject names...");
+
+				var ProjectPaths = UnrealVSPackage.Instance.GetLoadedProjectPaths();
+				var ProjectNames = (from path in ProjectPaths select Path.GetFileNameWithoutExtension(path)).ToArray();
+
+				var CodeUProjects = from UProjectPath in CachedUProjectPaths
+									let ProjectName = Path.GetFileNameWithoutExtension(UProjectPath)
+									where ProjectNames.Any(name => string.Compare(name, ProjectName, StringComparison.OrdinalIgnoreCase) == 0)
+									select new { Name = ProjectName, FilePath = UProjectPath };
+
+				CachedUProjects = new Dictionary<string, string>();
+
+				foreach (var UProject in CodeUProjects)
+				{
+					if (!CachedUProjects.ContainsKey(UProject.Name))
+					{
+						CachedUProjects.Add(UProject.Name, UProject.FilePath);
+					}
+				}
+
+				Logging.WriteLine("    DONE");
+			}
+
+			return CachedUProjects;
+		}
+
+		public static void GetSolutionConfigsAndPlatforms(out string[] SolutionConfigs, out string[] SolutionPlatforms)
+		{
+			var UniqueConfigs = new List<string>();
+			var UniquePlatforms = new List<string>();
+
+			SolutionConfigurations DteSolutionConfigs = UnrealVSPackage.Instance.DTE.Solution.SolutionBuild.SolutionConfigurations;
+			foreach (SolutionConfiguration2 SolutionConfig in DteSolutionConfigs)
+			{
+				if (!UniqueConfigs.Contains(SolutionConfig.Name))
+				{
+					UniqueConfigs.Add(SolutionConfig.Name);
+				}
+				if (!UniquePlatforms.Contains(SolutionConfig.PlatformName))
+				{
+					UniquePlatforms.Add(SolutionConfig.PlatformName);
+				}
+			}
+
+			SolutionConfigs = UniqueConfigs.ToArray();
+			SolutionPlatforms = UniquePlatforms.ToArray();
+		}
+
+		public static bool SetActiveSolutionConfiguration(string ConfigName, string PlatformName)
+		{
+			SolutionConfigurations DteSolutionConfigs = UnrealVSPackage.Instance.DTE.Solution.SolutionBuild.SolutionConfigurations;
+			foreach (SolutionConfiguration2 SolutionConfig in DteSolutionConfigs)
+			{
+				if (string.Compare(SolutionConfig.Name, ConfigName, StringComparison.Ordinal) == 0
+					&& string.Compare(SolutionConfig.PlatformName, PlatformName, StringComparison.Ordinal) == 0)
+				{
+					SolutionConfig.Activate();
+					return true;
+				}
+			}
+			return false;
+		}
+
+		public static bool SelectProjectInSolutionExplorer(Project Project)
+		{
+			UnrealVSPackage.Instance.DTE.ExecuteCommand("View.SolutionExplorer");
+			if (Project.ParentProjectItem != null)
+			{
+				Project.ParentProjectItem.ExpandView();
+			}
+
+			UIHierarchy SolutionExplorerHierarachy = UnrealVSPackage.Instance.DTE2.ToolWindows.SolutionExplorer;
+			Utils.UITreeItem SolutionExplorerTree = Utils.GetUIHierarchyTree(SolutionExplorerHierarachy);
+			var UIHierarachyProjects = Utils.GetUITreeItemsByObjectType<Project>(SolutionExplorerTree);
+
+			var SelectableUIItem = UIHierarachyProjects.FirstOrDefault(uihp => uihp.Object as Project == Project);
+
+			if (SelectableUIItem != null)
+			{
+				if (Project.ParentProjectItem != null)
+				{
+					SelectableUIItem.Select(vsUISelectionType.vsUISelectionTypeSelect);
+					return true;
+				}
+			}
+
+			return false;
+		}
+
+		public static void OnProjectListChanged()
+		{
+			CachedUProjects = null;
+		}
+
+		private static void PrepareOutputPane()
+		{
+			UnrealVSPackage.Instance.DTE.ExecuteCommand("View.Output");
+
+			var Pane = UnrealVSPackage.Instance.GetOutputPane();
+			if (Pane != null)
+			{
+				// Clear and activate the output pane.
+				Pane.Clear();
+
+				// @todo: Activating doesn't seem to really bring the pane to front like we would expect it to.
+				Pane.Activate();
+			}
+		}
+
+		/// Called by GetAllProjectsFromDTE() to list items from the project tree
+		private static void GetSubProjectsOfProjectItem(ProjectItem Item, List<Project> Projects)
+		{
+			if (Item.SubProject != null)
+			{
+				Projects.Add(Item.SubProject);
+
+				if (Item.SubProject.ProjectItems != null)
+				{
+					foreach (ProjectItem SubItem in Item.SubProject.ProjectItems)
+					{
+						GetSubProjectsOfProjectItem(SubItem, Projects);
+					}
+				}
+			}
+			if (Item.ProjectItems != null)
+			{
+				foreach (ProjectItem SubItem in Item.ProjectItems)
+				{
+					GetSubProjectsOfProjectItem(SubItem, Projects);
+				}
+			}
+		}
+
+		private static string GetSolutionFolder()
+		{
+			if (!UnrealVSPackage.Instance.DTE.Solution.IsOpen)
+			{
+				return string.Empty;
+			}
+
+			return Path.GetDirectoryName(UnrealVSPackage.Instance.SolutionFilepath);
+		}
+
+		private static string CachedUProjectRootFolder = string.Empty;
+		private static IEnumerable<string> CachedUProjectPaths = new string[0];
+		private static IDictionary<string, string> CachedUProjects = null;
+
+		private static string CachedUBTConfigFileName = string.Empty;
+		private static XmlNode CachedUBTConfigXml = null;
+	}
+}
diff --git a/Engine/Source/Programs/UnrealVS/UnrealVSOptions.cs b/Engine/Source/Programs/UnrealVS/UnrealVSOptions.cs
deleted file mode 100644
diff --git a/Engine/Source/Programs/UnrealVS/UnrealVSPackage.cs b/Engine/Source/Programs/UnrealVS/UnrealVSPackage.cs
deleted file mode 100644
diff --git a/Engine/Source/Programs/UnrealVS/UnrealVS_2019.csproj b/Engine/Source/Programs/UnrealVS/UnrealVS_2019.csproj
deleted file mode 100644
diff --git a/Engine/Source/Programs/UnrealVS/Utils.cs b/Engine/Source/Programs/UnrealVS/Utils.cs
deleted file mode 100644
diff --git a/Engine/Source/Programs/UnrealVS/app.config b/Engine/Source/Programs/UnrealVS/app.config
deleted file mode 100644
