[UnealVS] - add UnrealVS.P4 menu, P4Checkout command, optional Checkout on Save + framework for further additions

diff --git a/Engine/Source/Programs/UnrealVS/UnrealVS.Shared/P4Commands.cs b/Engine/Source/Programs/UnrealVS/UnrealVS.Shared/P4Commands.cs
new file mode 100644
--- /dev/null
+++ b/Engine/Source/Programs/UnrealVS/UnrealVS.Shared/P4Commands.cs
@@ -0,0 +1,408 @@
+// Copyright Epic Games, Inc. All Rights Reserved.
+
+using EnvDTE;
+using EnvDTE80;
+using Microsoft.VisualStudio;
+using Microsoft.VisualStudio.OLE.Interop;
+using Microsoft.VisualStudio.Shell;
+using Microsoft.VisualStudio.Shell.Interop;
+using System;
+using System.Collections.Generic;
+using System.ComponentModel.Design;
+using System.Diagnostics;
+using System.IO;
+using System.Windows.Forms;
+using Microsoft.Win32;
+
+namespace UnrealVS
+{
+	internal class P4Commands : IDisposable
+	{
+		private const int P4SubMenuID = 0x3100;
+		private const int P4CheckoutButtonID = 0x1450;
+
+		private OleMenuCommand SubMenuCommand;
+
+		private System.Diagnostics.Process ChildProcess;
+		private IVsOutputWindowPane P4OutputPane;
+		private string P4WorkingDirectory;
+		private bool PullWorkignDirectoryFromP4 = true;
+
+		// stdXX from last operation
+		private string P4OperationStdOut;
+		private string P4OperationStdErr;
+
+		// potentially pull this from registry - assume it exists there.
+		private string P4Exe = "C:\\Program Files\\Perforce\\p4.exe";
+
+		private List<CommandEvents> EventsForce = new List<CommandEvents>();
+
+		private class P4Command
+		{
+			public int ButtonID;
+
+			public MenuCommand ButtonCommand;
+			public CommandID CommandID;
+			public P4Command(int ButtonID, EventHandler ButtonHandler)
+			{
+				CommandID = new CommandID(GuidList.UnrealVSCmdSet, ButtonID);
+				ButtonCommand = new MenuCommand(new EventHandler(ButtonHandler), CommandID);
+
+				UnrealVSPackage.Instance.MenuCommandService.AddCommand(ButtonCommand);
+			}
+
+			public void Toggle(bool Enabled)
+			{
+				ButtonCommand.Visible = ButtonCommand.Enabled = Enabled;
+			}
+		}
+
+		private List<P4Command> P4CommandsList = new List<P4Command>();
+
+		private bool IsSolutionLoaded()
+		{
+			DTE DTE = UnrealVSPackage.Instance.DTE;
+
+			return DTE.Solution.FileName.Length > 0;
+		}
+
+		private void OnQuickBuildSubMenuQuery(object sender, EventArgs e)
+		{
+			DTE DTE = UnrealVSPackage.Instance.DTE;
+
+			bool EnableCommands = UnrealVSPackage.Instance.OptionsPage.AllowUnrealVSP4;
+			bool SolutionLoaded = IsSolutionLoaded();
+
+			var SenderSubMenuCommand = (OleMenuCommand)sender;
+
+			SenderSubMenuCommand.Visible = SenderSubMenuCommand.Enabled = EnableCommands & SolutionLoaded;
+		}
+
+		public P4Commands()
+		{
+			// setup callbacks on IDE operations
+			UnrealVSPackage.Instance.OptionsPage.OnOptionsChanged += OnOptionsChanged;
+			UnrealVSPackage.Instance.OnSolutionOpened += UpdateP4SoltuionBinding;
+			UnrealVSPackage.Instance.OnSolutionClosed += UpdateP4SoltuionBinding;
+
+			// create specific output window for unrealvs.P4
+			P4OutputPane = UnrealVSPackage.Instance.GetP4OutputPane();
+
+			// add commands
+			P4CommandsList.Add(new P4Command(P4CheckoutButtonID, P4CheckoutButtonHandler));
+
+			// add sub menu for commands
+			SubMenuCommand = new OleMenuCommand(null, new CommandID(GuidList.UnrealVSCmdSet, P4SubMenuID));
+			SubMenuCommand.BeforeQueryStatus += OnQuickBuildSubMenuQuery;
+			UnrealVSPackage.Instance.MenuCommandService.AddCommand(SubMenuCommand);
+
+			// Update the menu visibility to enforce user options.
+			UpdateMenuOptions();
+		}
+		// Called when solutions are loaded or unloaded
+		private void UpdateP4SoltuionBinding()
+		{
+			// Update the menu visibility
+			UpdateMenuOptions();
+
+			DTE DTE = UnrealVSPackage.Instance.DTE;
+
+			if (IsSolutionLoaded())
+			{
+				// use the current solution folder as a temp working directory
+				P4WorkingDirectory = Path.GetDirectoryName(DTE.Solution.FileName);
+
+				// force the system to pull the real working directory from p4
+				PullWorkignDirectoryFromP4 = true;
+			}
+			else
+			{
+				// reset the working directory to avoid sending commands to old streams
+				P4WorkingDirectory = "";
+			}
+		}
+
+		void RegisterCallbackHandler(string CommandName, _dispCommandEvents_BeforeExecuteEventHandler Callback)
+		{
+			// Find the command from the passed in name
+			DTE DTE = UnrealVSPackage.Instance.DTE;
+			CommandEvents Event = null;
+			{
+				// Probably should move this out to a function.
+				try
+				{
+					Command Command = DTE.Commands.Item(CommandName, -1);
+					if (Command != null)
+					{
+						Event = DTE.Events.get_CommandEvents(Command.Guid, Command.ID);
+					}
+				}
+				catch
+				{
+
+				}
+			}
+		
+			if (Event != null)
+			{
+				Event.BeforeExecute += Callback;
+				EventsForce.Add(Event); // forces a reference
+			}
+		}
+
+		void SaveSelectedCallback(string Guid, int ID, object CustomIn, object CustomOut, ref bool CancelDefault)
+		{
+			if (UnrealVSPackage.Instance.OptionsPage.AllowUnrealVSCheckoutOnEdit)
+			{
+				DTE DTE = UnrealVSPackage.Instance.DTE;
+
+				foreach (SelectedItem Item in DTE.SelectedItems)
+				{
+					if (Item.Project != null)
+					{
+						OpenForEdit(Item.Project.FullName);
+					}
+					else if (Item.ProjectItem != null)
+					{
+						OpenForEdit(Item.ProjectItem.Document.FullName);
+					}
+					else
+					{
+						OpenForEdit(DTE.Solution.FullName);
+					}
+				}
+			}
+		}
+
+		private void OnSaveAll(string Guid, int ID, object CustomIn, object CustomOut, ref bool CancelDefault)
+		{
+			if (UnrealVSPackage.Instance.OptionsPage.AllowUnrealVSCheckoutOnEdit)
+			{
+				DTE DTE = UnrealVSPackage.Instance.DTE;
+
+				if (!DTE.Solution.Saved)
+				{
+					OpenForEdit(DTE.Solution.FullName);
+				}
+
+				foreach (Document OpenedFile in DTE.Documents)
+				{
+					if (OpenedFile.Saved == false)
+					{
+						OpenForEdit(OpenedFile.FullName);
+					}
+				}
+
+				if (DTE.Solution.Projects == null)
+				{
+					return;
+				}
+
+				foreach (Project p in DTE.Solution.Projects)
+				{
+					EditProjectRecursive(p);
+				}
+			}
+		}
+
+		private void EditProjectRecursive(Project Proj)
+		{
+			if (!Proj.Saved)
+			{
+				OpenForEdit(Proj.FullName);
+			}
+				
+			if (Proj.ProjectItems == null)
+			{
+				return;
+			}
+
+			foreach (ProjectItem ProjItem in Proj.ProjectItems)
+			{
+				if (ProjItem.SubProject != null)
+				{
+					EditProjectRecursive(ProjItem.SubProject);
+				}
+				else if (!ProjItem.Saved)
+				{
+					for (short Index = 1; Index <= ProjItem.FileCount; Index++)
+					{
+						OpenForEdit(ProjItem.get_FileNames(Index));
+					}
+				}
+			}
+		}
+
+		private void UpdateMenuOptions()
+		{
+			DTE DTE = UnrealVSPackage.Instance.DTE;
+
+			bool EnableCommands = UnrealVSPackage.Instance.OptionsPage.AllowUnrealVSP4;
+			bool SolutionLoaded = IsSolutionLoaded();
+
+			// update each menu item enabled command
+			foreach (P4Command command in P4CommandsList)
+			{
+				command.Toggle(SolutionLoaded && EnableCommands);
+			}
+
+			// hook up the on "Save" list - internally they verify the users choice
+			RegisterCallbackHandler("File.SaveSelectedItems", SaveSelectedCallback);
+			RegisterCallbackHandler("File.SaveAll", OnSaveAll);
+		}
+		private void OnOptionsChanged(object Sender, EventArgs E)
+		{
+			UpdateMenuOptions();
+		}
+
+		public void Dispose()
+		{
+			KillChildProcess();
+		}
+
+		private void KillChildProcess()
+		{
+			if (ChildProcess != null)
+			{
+				if (!ChildProcess.HasExited)
+				{
+					ChildProcess.Kill();
+					ChildProcess.WaitForExit();
+				}
+				ChildProcess.Dispose();
+				ChildProcess = null;
+			}
+		}
+
+		private void P4CheckoutButtonHandler(object Sender, EventArgs Args)
+        {
+			DTE DTE = UnrealVSPackage.Instance.DTE;
+
+			// Check we've got a file open
+			if (DTE.ActiveDocument == null)
+			{
+				Logging.WriteLine("P4Checkout called without an active document");
+				return;
+			}
+
+			OpenForEdit(DTE.ActiveDocument.FullName);	
+		}
+
+		private void OpenForEdit(string FileName)
+		{
+			TryP4Command($"edit \"{FileName}");
+		}
+
+		private void ReadP4UserInfo()
+		{
+			//------P4 Operation started
+			//P4CLIENT = andrew.firth_private_frosty2(config 'd:\p4\frosty\p4config.txt')
+			//P4CONFIG = p4config.txt(set)(config 'd:\p4\frosty\p4config.txt')
+			//P4EDITOR = C:\windows\SysWOW64\notepad.exe(set)
+			//P4PORT = perforce:1666(config 'd:\p4\frosty\p4config.txt')
+			//P4USER = andrew.firth(config 'd:\p4\frosty\p4config.txt')
+			//P4_perforce:1666_CHARSET = none(set)
+
+			TryP4Command("set", true);
+
+			bool Success = false;
+
+			if (P4OperationStdOut.Length > 0)
+			{
+				string[] lines = P4OperationStdOut.Split('\n');
+
+				// very brittle - Index assume there is a better way leveraging C#
+				if (lines.Length > 0)
+				{
+					if (lines[0].Contains("p4config.txt"))
+					{
+						string path = lines[0].Split('\'')[1];
+
+						P4WorkingDirectory = Path.GetDirectoryName(path);
+						P4OutputPane.OutputString($"P4WorkingDirectory set to {P4WorkingDirectory}{Environment.NewLine}");
+						Success = true;
+					}
+				}
+			}
+
+			if (!Success)
+			{
+				P4OutputPane.OutputString($"attempt to full p4config.txt info failed, P4WorkingDirectory set to {P4WorkingDirectory}{Environment.NewLine}");
+			}
+		}
+
+		private void PullWorkingDirectory()
+		{
+			if (PullWorkignDirectoryFromP4)
+			{
+				ReadP4UserInfo();
+				PullWorkignDirectoryFromP4 = false;
+			}
+		}
+
+		private bool TryP4Command(string CommandLine, bool SkipPullWD = false, bool OutputStdOut = true)
+		{
+			if (SkipPullWD != true)
+			{
+				PullWorkingDirectory();
+			}
+		
+			DTE DTE = UnrealVSPackage.Instance.DTE;
+
+			// If there's already a operation in progress, abort... potential issues here
+			// in the OnSave/OnSaveAll operations.
+			if (ChildProcess != null && !ChildProcess.HasExited)
+			{
+				P4OutputPane.OutputString($"P4 operation already in flight.{Environment.NewLine}");
+				P4OutputPane.Activate();
+				return false;
+			}
+
+			// Make sure any existing op is stopped
+			KillChildProcess();
+
+			// Set up the output pane
+			P4OutputPane.Activate();
+			P4OutputPane.OutputString($"1>------ P4 Operation started{Environment.NewLine}");
+			P4OutputPane.OutputString($"CMD: {P4Exe} {CommandLine}{Environment.NewLine}");
+
+			// Create a delegate for handling output messages
+			void OutputHandler(object Sender, DataReceivedEventArgs Args) { if (Args.Data != null) P4OperationStdOut += $"{Args.Data}{Environment.NewLine}"; }
+			void ErrOutputHandler(object Sender, DataReceivedEventArgs Args) { if (Args.Data != null) P4OperationStdErr += $"{Args.Data}{Environment.NewLine}"; }
+
+			P4OperationStdOut = "";
+			P4OperationStdErr = "";
+
+			// Spawn the new process
+			ChildProcess = new System.Diagnostics.Process();
+			ChildProcess.StartInfo.FileName = P4Exe;
+			ChildProcess.StartInfo.Arguments = CommandLine;
+			ChildProcess.StartInfo.WorkingDirectory = P4WorkingDirectory;
+			ChildProcess.StartInfo.UseShellExecute = false;
+			ChildProcess.StartInfo.RedirectStandardOutput = true;
+			ChildProcess.StartInfo.RedirectStandardError = true;
+			ChildProcess.StartInfo.CreateNoWindow = true;
+			ChildProcess.OutputDataReceived += OutputHandler;
+			ChildProcess.ErrorDataReceived += ErrOutputHandler;
+			ChildProcess.Start();
+			ChildProcess.BeginOutputReadLine();
+			ChildProcess.BeginErrorReadLine();
+
+			ChildProcess.WaitForExit();
+
+			if (P4OperationStdOut.Length > 0 && OutputStdOut)
+			{
+				P4OutputPane.OutputString(P4OperationStdOut);
+			}
+
+			if (P4OperationStdErr.Length > 0)
+			{
+				P4OutputPane.OutputString(P4OperationStdErr);
+			}
+
+			P4OutputPane.OutputString($"1>------ P4 Operation complete{Environment.NewLine}");
+
+			return true;
+		}
+	}
+}
diff --git a/Engine/Source/Programs/UnrealVS/UnrealVS.Shared/UnrealVSOptions.cs b/Engine/Source/Programs/UnrealVS/UnrealVS.Shared/UnrealVSOptions.cs
--- a/Engine/Source/Programs/UnrealVS/UnrealVS.Shared/UnrealVSOptions.cs
+++ b/Engine/Source/Programs/UnrealVS/UnrealVS.Shared/UnrealVSOptions.cs
@@ -23,6 +23,27 @@ namespace UnrealVS
 			set { _HideNonGameStartupProjects = value; }
 		}
 
+		private bool _AllowUnrealVSP4;
+		private bool _AllowUnrealVSCheckoutOnEdit;
+
+		[Category("General")]
+		[DisplayName("Allow UnrealVS to auto checkout files on save")]
+		[Description("Uses p4 ini / environment settings to automatically checkout files on save, use 'RunUAT P4WriteConfig' to initialize those settings")]
+		public bool AllowUnrealVSCheckoutOnEdit
+		{
+			get { return _AllowUnrealVSCheckoutOnEdit; }
+			set { _AllowUnrealVSCheckoutOnEdit = value; }
+		}
+
+		[Category("General")]
+		[DisplayName("Allow UnrealVS to perform perforce operations")]
+		[Description("Uses p4 ini / environment settings to call P4 functionlality, use 'RunUAT P4WriteConfig' to initialize those settings")]
+		public bool AllowUnrealVSP4
+		{
+			get { return _AllowUnrealVSP4; }
+			set { _AllowUnrealVSP4 = value; }
+		}
+
 		protected override void OnApply(PageApplyEventArgs e)
 		{
 			base.OnApply(e);
diff --git a/Engine/Source/Programs/UnrealVS/UnrealVS.Shared/UnrealVSPackage.cs b/Engine/Source/Programs/UnrealVS/UnrealVS.Shared/UnrealVSPackage.cs
--- a/Engine/Source/Programs/UnrealVS/UnrealVS.Shared/UnrealVSPackage.cs
+++ b/Engine/Source/Programs/UnrealVS/UnrealVS.Shared/UnrealVSPackage.cs
@@ -239,6 +239,9 @@ namespace UnrealVS
 			// Create the project menu quick builder
 			QuickBuilder = new QuickBuild();
 
+			// Create 'Perforce menu' instance
+			P4CommandsGroup = new P4Commands();
+
 			// Call parent implementation
 			base.Initialize();
 
@@ -246,6 +249,7 @@ namespace UnrealVS
 			{
 				StartTicker();
 			}
+
 		}
 
 		private void StartTicker()
@@ -485,6 +489,15 @@ namespace UnrealVS
 			return GetOutputPane(VSConstants.OutputWindowPaneGuid.BuildOutputPane_guid, "Build");
 		}
 
+		/// <summary>
+		/// Gets a Visual Studio pane to output text to for P4 ops, or creates one if not visible.  Does not bring the pane to front (you can call Activate() to do that.)
+		/// </summary>
+		/// <returns>The pane to output to, or null on error</returns>
+		public IVsOutputWindowPane GetP4OutputPane()
+		{
+			return GetOutputPane(VSConstants.OutputWindowPaneGuid.GeneralPane_guid, "UnrealVS.P4");
+		}
+
 		public IEnumerable<string> GetLoadedProjectPaths()
 		{
 			return LoadedProjectPaths;
@@ -922,6 +935,9 @@ namespace UnrealVS
 		/// Batch Builder button/command handler
 		private BatchBuilder BatchBuilder = new BatchBuilder();
 
+		/// Perforce features
+		private P4Commands P4CommandsGroup;
+
 		/// Ticker thread
 		private Thread Ticker;
 
