[UnrealVS] - add annotate command

- this is a basic form that will hit issues where our integration patterns somewhat reduce its efficacy
- version 2 will improve on this

diff --git a/Engine/Source/Programs/UnrealVS/UnrealVS.Shared/P4Commands.cs b/Engine/Source/Programs/UnrealVS/UnrealVS.Shared/P4Commands.cs
--- a/Engine/Source/Programs/UnrealVS/UnrealVS.Shared/P4Commands.cs
+++ b/Engine/Source/Programs/UnrealVS/UnrealVS.Shared/P4Commands.cs
@@ -18,8 +18,15 @@ namespace UnrealVS
 {
 	internal class P4Commands : IDisposable
 	{
+		private const bool PullWorkingDirectoryOn = true;
+		private const bool PullWorkingDirectoryOff = false;
+
+		private const bool OutputStdOutOn= true;
+		private const bool OutputStdOutOff = false;
+
 		private const int P4SubMenuID = 0x3100;
 		private const int P4CheckoutButtonID = 0x1450;
+		private const int P4AnnotateButtonID = 0x1451;
 
 		private OleMenuCommand SubMenuCommand;
 
@@ -90,6 +97,7 @@ namespace UnrealVS
 
 			// add commands
 			P4CommandsList.Add(new P4Command(P4CheckoutButtonID, P4CheckoutButtonHandler));
+			P4CommandsList.Add(new P4Command(P4AnnotateButtonID, P4AnnotateButtonHandler));
 
 			// add sub menu for commands
 			SubMenuCommand = new OleMenuCommand(null, new CommandID(GuidList.UnrealVSCmdSet, P4SubMenuID));
@@ -282,11 +290,108 @@ namespace UnrealVS
 			if (DTE.ActiveDocument == null)
 			{
 				Logging.WriteLine("P4Checkout called without an active document");
+
+				P4OutputPane.Activate();
+				P4OutputPane.OutputString($"1>------ P4Checkout called without an active document{Environment.NewLine}");
+
 				return;
 			}
 
 			OpenForEdit(DTE.ActiveDocument.FullName);	
 		}
+		private void P4AnnotateButtonHandler(object Sender, EventArgs Args)
+		{
+			DTE DTE = UnrealVSPackage.Instance.DTE;
+
+			P4OutputPane.Activate();
+
+			// Check we've got a file open
+			if (DTE.ActiveDocument == null)
+			{
+				Logging.WriteLine("P4Annotate called without an active document");
+
+				P4OutputPane.OutputString($"1>------ P4Annotate called without an active document{Environment.NewLine}");
+
+				return;
+			}
+
+			// Call annotate itself
+			bool Result= TryP4Command($"annotate -TcIqu \"{DTE.ActiveDocument.FullName}", PullWorkingDirectoryOn, OutputStdOutOff);
+
+			if (!Result || P4OperationStdErr.Length>0)
+			{
+				P4OutputPane.OutputString($"1>------ P4Annotate call failed");
+				return;
+			}
+
+			// Extract the current document line number and the first line of text within it
+			string FirstLine;
+			int CurrentLine = 0;
+			{
+				TextDocument CurrentDocument = (TextDocument)(DTE.ActiveDocument.Object("TextDocument"));
+
+				var EditPoint = CurrentDocument.StartPoint.CreateEditPoint();
+				string DocumentContents = EditPoint.GetText(CurrentDocument.EndPoint);
+				FirstLine = DocumentContents.Split('\n')[0];
+
+
+				TextSelection TextSel = DTE.ActiveWindow.Selection as TextSelection;
+				if (TextSel != null)
+				{
+					CurrentLine = TextSel.CurrentLine;
+				}
+			}
+
+			// Pre-process the output to comment out the additions thus allowing
+			// code to use correct syntax coloring - helps enormously with visualization
+			string EditedCopyOfAnnotate = "";
+			{
+				// replace 
+				//       13149436:            First.Last 2020/05/04 
+				//
+				// with
+				// /*	13149436:            First.Last 2020/05/04*/
+
+				// This is the per line offset of the annotation added to the document
+				int AnnotateOffset = P4OperationStdOut.IndexOf(FirstLine);
+
+				string[] AnnotateLines = P4OperationStdOut.Split('\n');
+
+				foreach (string Line in AnnotateLines)
+				{
+					if (Line.Length > AnnotateOffset)
+					{
+						string EditedLine = Line.Insert(AnnotateOffset, "*/");
+						EditedLine = EditedLine.Insert(0, "/*");
+
+						EditedCopyOfAnnotate += EditedLine;
+					}
+					else
+					{
+						EditedCopyOfAnnotate += Line;
+					}
+
+				}
+			}
+
+			// Replace GetTempPath with the UBT intermediate folder if we have access
+			string TempPath = Path.GetTempPath();
+			string TempFileName = Path.GetFileNameWithoutExtension(DTE.ActiveDocument.FullName) + "_annotate" + Path.GetExtension(DTE.ActiveDocument.FullName);
+			string TempFilePath = Path.Combine(TempPath, TempFileName);
+
+			// Write out our temp file
+			File.WriteAllText(TempFilePath, EditedCopyOfAnnotate);
+
+			// Open it, activate it and move to the line the user focused to execute the command
+			DTE.ExecuteCommand("File.OpenFile", $"\"{TempFilePath}\"");
+			DTE.ActiveDocument.Activate();
+
+			TextSelection NewTextSel = DTE.ActiveWindow.Selection as TextSelection;
+			if (NewTextSel != null)
+			{
+				NewTextSel.GotoLine(CurrentLine, false);
+			}
+		}
 
 		private void OpenForEdit(string FileName)
 		{
@@ -303,7 +408,7 @@ namespace UnrealVS
 			//P4USER = andrew.firth(config 'd:\p4\frosty\p4config.txt')
 			//P4_perforce:1666_CHARSET = none(set)
 
-			TryP4Command("set", true);
+			TryP4Command("set", PullWorkingDirectoryOff);
 
 			bool Success = false;
 
@@ -340,9 +445,9 @@ namespace UnrealVS
 			}
 		}
 
-		private bool TryP4Command(string CommandLine, bool SkipPullWD = false, bool OutputStdOut = true)
+		private bool TryP4Command(string CommandLine, bool PullWorkingDirectoryNow = PullWorkingDirectoryOn, bool OutputStdOut = OutputStdOutOn)
 		{
-			if (SkipPullWD != true)
+			if (PullWorkingDirectoryNow)
 			{
 				PullWorkingDirectory();
 			}
