[UnrealVS] - respond to C# feedback

diff --git a/Engine/Source/Programs/UnrealVS/UnrealVS.Shared/P4Commands.cs b/Engine/Source/Programs/UnrealVS/UnrealVS.Shared/P4Commands.cs
--- a/Engine/Source/Programs/UnrealVS/UnrealVS.Shared/P4Commands.cs
+++ b/Engine/Source/Programs/UnrealVS/UnrealVS.Shared/P4Commands.cs
@@ -11,6 +11,7 @@ using System.Collections.Generic;
 using System.ComponentModel.Design;
 using System.Diagnostics;
 using System.IO;
+using System.Text;
 using System.Windows.Forms;
 using Microsoft.Win32;
 
@@ -344,7 +345,7 @@ namespace UnrealVS
 
 			// Pre-process the output to comment out the additions thus allowing
 			// code to use correct syntax coloring - helps enormously with visualization
-			string EditedCopyOfAnnotate = "";
+			StringBuilder EditedCopy = new StringBuilder();
 			{
 				// replace 
 				//       13149436:            First.Last 2020/05/04 
@@ -364,11 +365,11 @@ namespace UnrealVS
 						string EditedLine = Line.Insert(AnnotateOffset, "*/");
 						EditedLine = EditedLine.Insert(0, "/*");
 
-						EditedCopyOfAnnotate += EditedLine;
+						EditedCopy.Append(EditedLine);
 					}
 					else
 					{
-						EditedCopyOfAnnotate += Line;
+						EditedCopy.Append(Line);
 					}
 
 				}
@@ -380,7 +381,7 @@ namespace UnrealVS
 			string TempFilePath = Path.Combine(TempPath, TempFileName);
 
 			// Write out our temp file
-			File.WriteAllText(TempFilePath, EditedCopyOfAnnotate);
+			File.WriteAllText(TempFilePath, EditedCopy.ToString());
 
 			// Open it, activate it and move to the line the user focused to execute the command
 			DTE.ExecuteCommand("File.OpenFile", $"\"{TempFilePath}\"");
@@ -472,30 +473,39 @@ namespace UnrealVS
 			P4OutputPane.OutputString($"CMD: {P4Exe} {CommandLine}{Environment.NewLine}");
 
 			// Create a delegate for handling output messages
-			void OutputHandler(object Sender, DataReceivedEventArgs Args) { if (Args.Data != null) P4OperationStdOut += $"{Args.Data}{Environment.NewLine}"; }
-			void ErrOutputHandler(object Sender, DataReceivedEventArgs Args) { if (Args.Data != null) P4OperationStdErr += $"{Args.Data}{Environment.NewLine}"; }
+			//void OutputHandler(object Sender, DataReceivedEventArgs Args) { if (Args.Data != null) P4OperationStdOut += $"{Args.Data}{Environment.NewLine}"; }
+			//void ErrOutputHandler(object Sender, DataReceivedEventArgs Args) { if (Args.Data != null) P4OperationStdErr += $"{Args.Data}{Environment.NewLine}"; }
 
-			P4OperationStdOut = "";
-			P4OperationStdErr = "";
+			StringBuilder StdOutSB = new StringBuilder();
+			StringBuilder StdErrSB = new StringBuilder();
 
 			// Spawn the new process
-			ChildProcess = new System.Diagnostics.Process();
-			ChildProcess.StartInfo.FileName = P4Exe;
-			ChildProcess.StartInfo.Arguments = CommandLine;
-			ChildProcess.StartInfo.WorkingDirectory = P4WorkingDirectory;
-			ChildProcess.StartInfo.UseShellExecute = false;
-			ChildProcess.StartInfo.RedirectStandardOutput = true;
-			ChildProcess.StartInfo.RedirectStandardError = true;
-			ChildProcess.StartInfo.CreateNoWindow = true;
-			ChildProcess.OutputDataReceived += OutputHandler;
-			ChildProcess.ErrorDataReceived += ErrOutputHandler;
+			ChildProcess = new System.Diagnostics.Process()
+			{
+				StartInfo = new ProcessStartInfo()
+				{
+					FileName = P4Exe,
+					Arguments = CommandLine,
+					WorkingDirectory = P4WorkingDirectory,
+					UseShellExecute = false,
+					RedirectStandardOutput = true,
+					RedirectStandardError = true,
+					CreateNoWindow = true
+				}
+			};
+			// Create a delegate for handling output messages
+			ChildProcess.OutputDataReceived += (s, a) => { if (a.Data != null) StdOutSB.AppendLine(a.Data); };
+			ChildProcess.ErrorDataReceived += (s, a) => { if (a.Data != null) StdErrSB.AppendLine(a.Data);	};
 			ChildProcess.Start();
 			ChildProcess.BeginOutputReadLine();
 			ChildProcess.BeginErrorReadLine();
 
 			ChildProcess.WaitForExit();
 
-			if (P4OperationStdOut.Length > 0 && OutputStdOut)
+			P4OperationStdOut = StdOutSB.ToString();
+			P4OperationStdErr = StdErrSB.ToString();
+
+if (P4OperationStdOut.Length > 0 && OutputStdOut)
 			{
 				P4OutputPane.OutputString(P4OperationStdOut);
 			}
