UnrealVS: Rebuild VS2022 extension for experimental p4 edit functionalily
Check if the file is read only before editing

#rnx

diff --git a/Engine/Source/Programs/UnrealVS/UnrealVS.Shared/P4Commands.cs b/Engine/Source/Programs/UnrealVS/UnrealVS.Shared/P4Commands.cs
--- a/Engine/Source/Programs/UnrealVS/UnrealVS.Shared/P4Commands.cs
+++ b/Engine/Source/Programs/UnrealVS/UnrealVS.Shared/P4Commands.cs
@@ -396,6 +396,12 @@ namespace UnrealVS
 
 		private void OpenForEdit(string FileName)
 		{
+			// Don't open for edit if the file is already writable
+			if (!File.Exists(FileName) || !File.GetAttributes(FileName).HasFlag(FileAttributes.ReadOnly))
+			{
+				return;
+			}
+
 			TryP4Command($"edit \"{FileName}");
 		}
 
