[UnrealVS] - fix issues with p4 setup, the CWD was hitting timing issues and the OnSave callbacks were being duped across solution load/unload

diff --git a/Engine/Source/Programs/UnrealVS/UnrealVS.2019/Properties/AssemblyInfo.cs b/Engine/Source/Programs/UnrealVS/UnrealVS.2019/Properties/AssemblyInfo.cs
--- a/Engine/Source/Programs/UnrealVS/UnrealVS.2019/Properties/AssemblyInfo.cs
+++ b/Engine/Source/Programs/UnrealVS/UnrealVS.2019/Properties/AssemblyInfo.cs
@@ -30,8 +30,8 @@ using System.Runtime.InteropServices;
 // You can specify all the values or you can default the Revision and Build Numbers 
 // by using the '*' as shown below:
 
-[assembly: AssemblyVersion("1.62.0.0")]
-[assembly: AssemblyFileVersion("1.62.0.0")]
+[assembly: AssemblyVersion("1.63.0.0")]
+[assembly: AssemblyFileVersion("1.63.0.0")]
 
 
 
diff --git a/Engine/Source/Programs/UnrealVS/UnrealVS.2022/Properties/AssemblyInfo.cs b/Engine/Source/Programs/UnrealVS/UnrealVS.2022/Properties/AssemblyInfo.cs
--- a/Engine/Source/Programs/UnrealVS/UnrealVS.2022/Properties/AssemblyInfo.cs
+++ b/Engine/Source/Programs/UnrealVS/UnrealVS.2022/Properties/AssemblyInfo.cs
@@ -30,8 +30,8 @@ using System.Runtime.InteropServices;
 // You can specify all the values or you can default the Revision and Build Numbers 
 // by using the '*' as shown below:
 
-[assembly: AssemblyVersion("1.62.0.0")]
-[assembly: AssemblyFileVersion("1.62.0.0")]
+[assembly: AssemblyVersion("1.63.0.0")]
+[assembly: AssemblyFileVersion("1.63.0.0")]
 
 
 
diff --git a/Engine/Source/Programs/UnrealVS/UnrealVS.Shared/P4Commands.cs b/Engine/Source/Programs/UnrealVS/UnrealVS.Shared/P4Commands.cs
--- a/Engine/Source/Programs/UnrealVS/UnrealVS.Shared/P4Commands.cs
+++ b/Engine/Source/Programs/UnrealVS/UnrealVS.Shared/P4Commands.cs
@@ -34,7 +34,7 @@ namespace UnrealVS
 		private System.Diagnostics.Process ChildProcess;
 		private IVsOutputWindowPane P4OutputPane;
 		private string P4WorkingDirectory;
-		private bool PullWorkignDirectoryFromP4 = true;
+		private bool PullWorkingDirectorFromP4 = true;
 
 		// stdXX from last operation
 		private string P4OperationStdOut;
@@ -90,8 +90,8 @@ namespace UnrealVS
 		{
 			// setup callbacks on IDE operations
 			UnrealVSPackage.Instance.OptionsPage.OnOptionsChanged += OnOptionsChanged;
-			UnrealVSPackage.Instance.OnSolutionOpened += UpdateP4SoltuionBinding;
-			UnrealVSPackage.Instance.OnSolutionClosed += UpdateP4SoltuionBinding;
+			UnrealVSPackage.Instance.OnSolutionOpened += SoltuionOpened;
+			UnrealVSPackage.Instance.OnSolutionClosed += SoltuionClosed;
 
 			// create specific output window for unrealvs.P4
 			P4OutputPane = UnrealVSPackage.Instance.GetP4OutputPane();
@@ -107,28 +107,26 @@ namespace UnrealVS
 
 			// Update the menu visibility to enforce user options.
 			UpdateMenuOptions();
+
+			// hook up the on "Save" list - internally they verify the users choice
+			RegisterCallbackHandler("File.SaveSelectedItems", SaveSelectedCallback);
+			RegisterCallbackHandler("File.SaveAll", OnSaveAll);
 		}
 		// Called when solutions are loaded or unloaded
-		private void UpdateP4SoltuionBinding()
+		private void SoltuionOpened()
 		{
 			// Update the menu visibility
 			UpdateMenuOptions();
+		}
 
-			DTE DTE = UnrealVSPackage.Instance.DTE;
-
-			if (IsSolutionLoaded())
-			{
-				// use the current solution folder as a temp working directory
-				P4WorkingDirectory = Path.GetDirectoryName(DTE.Solution.FileName);
+		private void SoltuionClosed()
+		{
+			// Update the menu visibility
+			UpdateMenuOptions();
 
-				// force the system to pull the real working directory from p4
-				PullWorkignDirectoryFromP4 = true;
-			}
-			else
-			{
-				// reset the working directory to avoid sending commands to old streams
-				P4WorkingDirectory = "";
-			}
+			// Clear any existing P4 working directory settings
+			P4WorkingDirectory = "";
+			PullWorkingDirectorFromP4 = true;
 		}
 
 		void RegisterCallbackHandler(string CommandName, _dispCommandEvents_BeforeExecuteEventHandler Callback)
@@ -255,9 +253,6 @@ namespace UnrealVS
 				command.Toggle(SolutionLoaded && EnableCommands);
 			}
 
-			// hook up the on "Save" list - internally they verify the users choice
-			RegisterCallbackHandler("File.SaveSelectedItems", SaveSelectedCallback);
-			RegisterCallbackHandler("File.SaveAll", OnSaveAll);
 		}
 		private void OnOptionsChanged(object Sender, EventArgs E)
 		{
@@ -405,8 +400,9 @@ namespace UnrealVS
 			TryP4Command($"edit \"{FileName}");
 		}
 
-		private void ReadP4UserInfo()
+		private string ReadWorkingDirectorFromP4Config()
 		{
+			string WorkingDirectory = "";
 			//------P4 Operation started
 			//P4CLIENT = andrew.firth_private_frosty2(config 'd:\p4\frosty\p4config.txt')
 			//P4CONFIG = p4config.txt(set)(config 'd:\p4\frosty\p4config.txt')
@@ -430,8 +426,7 @@ namespace UnrealVS
 					{
 						string path = lines[0].Split('\'')[1];
 
-						P4WorkingDirectory = Path.GetDirectoryName(path);
-						P4OutputPane.OutputString($"P4WorkingDirectory set to {P4WorkingDirectory}{Environment.NewLine}");
+						WorkingDirectory = Path.GetDirectoryName(path);
 						Success = true;
 					}
 				}
@@ -439,25 +434,47 @@ namespace UnrealVS
 
 			if (!Success)
 			{
-				P4OutputPane.OutputString($"attempt to full p4config.txt info failed, P4WorkingDirectory set to {P4WorkingDirectory}{Environment.NewLine}");
+				P4OutputPane.OutputString($"attempt to pull p4config.txt info failed{Environment.NewLine}");
 			}
+
+			return WorkingDirectory;
 		}
 
-		private void PullWorkingDirectory()
+		private void PullWorkingDirectory(bool pullfromP4Settings)
 		{
-			if (PullWorkignDirectoryFromP4)
+			if (IsSolutionLoaded() || P4WorkingDirectory != null || P4WorkingDirectory.Length < 2)
 			{
-				ReadP4UserInfo();
-				PullWorkignDirectoryFromP4 = false;
+				DTE DTE = UnrealVSPackage.Instance.DTE;
+
+				// use the current solution folder as a temp working directory
+				P4WorkingDirectory = Path.GetDirectoryName(DTE.Solution.FileName);
+			}
+
+			// if the callee wants us to pull a CWD and it wasn't already done
+			// pull it from the p4 config now
+			if (pullfromP4Settings && PullWorkingDirectorFromP4)
+			{
+				string NewWorkingDirectory = ReadWorkingDirectorFromP4Config();
+
+				if (NewWorkingDirectory.Length > 1)
+				{
+					P4WorkingDirectory = NewWorkingDirectory;
+					PullWorkingDirectorFromP4 = false;
+
+					P4OutputPane.OutputString($"P4WorkingDirectory set to '{P4WorkingDirectory}'{Environment.NewLine}");
+				}
+				else
+				{
+					P4OutputPane.OutputString($"P4WorkingDirectory set failed {Environment.NewLine}");
+				}
+
+				
 			}
 		}
 
 		private bool TryP4Command(string CommandLine, bool PullWorkingDirectoryNow = PullWorkingDirectoryOn, bool OutputStdOut = OutputStdOutOn)
 		{
-			if (PullWorkingDirectoryNow)
-			{
-				PullWorkingDirectory();
-			}
+			PullWorkingDirectory(PullWorkingDirectoryNow);
 		
 			DTE DTE = UnrealVSPackage.Instance.DTE;
 
@@ -476,6 +493,7 @@ namespace UnrealVS
 			// Set up the output pane
 			P4OutputPane.Activate();
 			P4OutputPane.OutputString($"1>------ P4 Operation started{Environment.NewLine}");
+			P4OutputPane.OutputString($"CWD: {P4WorkingDirectory}{Environment.NewLine}");
 			P4OutputPane.OutputString($"CMD: {P4Exe} {CommandLine}{Environment.NewLine}");
 
 			// Create a delegate for handling output messages
