[UnrealVS] - add support for P4VC, add commands: open selected CL, open integration aware timelapse view, and Diff within Visual Studio + Optionally override diff in VS settings

$rb none

diff --git a/Engine/Source/Programs/UnrealVS/UnrealVS.2019/Properties/AssemblyInfo.cs b/Engine/Source/Programs/UnrealVS/UnrealVS.2019/Properties/AssemblyInfo.cs
--- a/Engine/Source/Programs/UnrealVS/UnrealVS.2019/Properties/AssemblyInfo.cs
+++ b/Engine/Source/Programs/UnrealVS/UnrealVS.2019/Properties/AssemblyInfo.cs
@@ -30,8 +30,8 @@ using System.Runtime.InteropServices;
 // You can specify all the values or you can default the Revision and Build Numbers 
 // by using the '*' as shown below:
 
-[assembly: AssemblyVersion("1.63.0.0")]
-[assembly: AssemblyFileVersion("1.63.0.0")]
+[assembly: AssemblyVersion("1.64.0.0")]
+[assembly: AssemblyFileVersion("1.64.0.0")]
 
 
 
diff --git a/Engine/Source/Programs/UnrealVS/UnrealVS.2022/Properties/AssemblyInfo.cs b/Engine/Source/Programs/UnrealVS/UnrealVS.2022/Properties/AssemblyInfo.cs
--- a/Engine/Source/Programs/UnrealVS/UnrealVS.2022/Properties/AssemblyInfo.cs
+++ b/Engine/Source/Programs/UnrealVS/UnrealVS.2022/Properties/AssemblyInfo.cs
@@ -30,8 +30,8 @@ using System.Runtime.InteropServices;
 // You can specify all the values or you can default the Revision and Build Numbers 
 // by using the '*' as shown below:
 
-[assembly: AssemblyVersion("1.63.0.0")]
-[assembly: AssemblyFileVersion("1.63.0.0")]
+[assembly: AssemblyVersion("1.64.0.0")]
+[assembly: AssemblyFileVersion("1.64.0.0")]
 
 
 
diff --git a/Engine/Source/Programs/UnrealVS/UnrealVS.Shared/P4Commands.cs b/Engine/Source/Programs/UnrealVS/UnrealVS.Shared/P4Commands.cs
--- a/Engine/Source/Programs/UnrealVS/UnrealVS.Shared/P4Commands.cs
+++ b/Engine/Source/Programs/UnrealVS/UnrealVS.Shared/P4Commands.cs
@@ -1,19 +1,20 @@
 // Copyright Epic Games, Inc. All Rights Reserved.
 
 using EnvDTE;
-using EnvDTE80;
-using Microsoft.VisualStudio;
-using Microsoft.VisualStudio.OLE.Interop;
 using Microsoft.VisualStudio.Shell;
 using Microsoft.VisualStudio.Shell.Interop;
+using Microsoft.VisualStudio.TextManager.Interop;
 using System;
 using System.Collections.Generic;
 using System.ComponentModel.Design;
 using System.Diagnostics;
 using System.IO;
 using System.Text;
-using System.Windows.Forms;
-using Microsoft.Win32;
+using System.Text.RegularExpressions;
+using Microsoft.VisualStudio.Settings;
+using Microsoft.VisualStudio.Shell.Settings;
+using Microsoft.VisualStudio.Text.Editor;
+using Microsoft.VisualStudio.Text.Differencing;
 
 namespace UnrealVS
 {
@@ -28,6 +29,11 @@ namespace UnrealVS
 		private const int P4SubMenuID = 0x3100;
 		private const int P4CheckoutButtonID = 0x1450;
 		private const int P4AnnotateButtonID = 0x1451;
+		private const int P4ViewSelectedCLButtonID = 0x1452;
+		private const int P4IntegrationAwareTimelapseButtonID = 0x1453;
+		private const int P4DiffinVSButtonID = 0x1454;
+		
+
 
 		private OleMenuCommand SubMenuCommand;
 
@@ -40,15 +46,22 @@ namespace UnrealVS
 		private string P4OperationStdOut;
 		private string P4OperationStdErr;
 
-		// potentially pull this from registry - assume it exists there.
+		// Exe paths
 		private string P4Exe = "C:\\Program Files\\Perforce\\p4.exe";
+		private string P4VCCmd = "C:\\Program Files\\Perforce\\p4vc.exe";
+		private string P4VCCmdBat = "C:\\Program Files\\Perforce\\p4vc.bat";
+		private string P4VExe = "C:\\Program Files\\Perforce\\p4v.exe";
+
+		// user info
+		private string Username = "";
+		private string Port = "";
+		private string Client = "";
+		private string UserInfoComplete = "";
 
 		private List<CommandEvents> EventsForce = new List<CommandEvents>();
 
 		private class P4Command
 		{
-			public int ButtonID;
-
 			public MenuCommand ButtonCommand;
 			public CommandID CommandID;
 			public P4Command(int ButtonID, EventHandler ButtonHandler)
@@ -95,10 +108,34 @@ namespace UnrealVS
 
 			// create specific output window for unrealvs.P4
 			P4OutputPane = UnrealVSPackage.Instance.GetP4OutputPane();
+			
+			// figure out the P4VC path
+			if (!File.Exists(P4VCCmd))
+			{
+				if (File.Exists(P4VCCmd))
+				{
+					P4VCCmd = P4VCCmdBat;
+				}
+				else 
+				{
+					P4OutputPane.Activate();
+					P4OutputPane.OutputString($"1>------ P4VC not found, {P4VCCmd} or {P4VCCmd}{Environment.NewLine}");
+					P4VCCmd = "";
+				}
+					
+			}
 
 			// add commands
 			P4CommandsList.Add(new P4Command(P4CheckoutButtonID, P4CheckoutButtonHandler));
 			P4CommandsList.Add(new P4Command(P4AnnotateButtonID, P4AnnotateButtonHandler));
+			P4CommandsList.Add(new P4Command(P4IntegrationAwareTimelapseButtonID, P4IntegrationAwareTimeLapseHandler));
+			P4CommandsList.Add(new P4Command(P4DiffinVSButtonID, P4DiffinVSHandler));
+
+			if (P4VCCmd.Length > 1)
+			{	
+				P4CommandsList.Add(new P4Command(P4ViewSelectedCLButtonID, P4ViewSelectedCLButtonHandler));
+
+			}
 
 			// add sub menu for commands
 			SubMenuCommand = new OleMenuCommand(null, new CommandID(GuidList.UnrealVSCmdSet, P4SubMenuID));
@@ -111,6 +148,7 @@ namespace UnrealVS
 			// hook up the on "Save" list - internally they verify the users choice
 			RegisterCallbackHandler("File.SaveSelectedItems", SaveSelectedCallback);
 			RegisterCallbackHandler("File.SaveAll", OnSaveAll);
+			
 		}
 		// Called when solutions are loaded or unloaded
 		private void SoltuionOpened()
@@ -388,7 +426,106 @@ namespace UnrealVS
 				NewTextSel.GotoLine(CurrentLine, false);
 			}
 		}
+		private void P4ViewSelectedCLButtonHandler(object Sender, EventArgs Args)
+		{
+			DTE DTE = UnrealVSPackage.Instance.DTE;
+
+			if (DTE.ActiveDocument == null)
+			{
+				return;
+			}
+
+			TextSelection TextSel = (TextSelection)DTE.ActiveDocument.Selection;
+
+			int ChangeList = -1;
+
+			try
+			{
+				ChangeList = Int32.Parse(TextSel.Text);
+			}
+			catch
+			{
+				ChangeList = -2;
+			}
+			
+
+			if (ChangeList > 0)
+			{
+				TryP4VCCommand($"Change {ChangeList}");
+			}
+			
+		}
+
+		private void P4IntegrationAwareTimeLapseHandler(object Sender, EventArgs Args)
+		{
+			DTE DTE = UnrealVSPackage.Instance.DTE;
+
+			if (DTE.ActiveDocument == null)
+			{
+				Logging.WriteLine("P4IntegrationAwareTimeLapse called without an active document");
+
+				P4OutputPane.OutputString($"1>------ P4IntegrationAwareTimeLapse called without an active document{Environment.NewLine}");
+
+				return;
+			}
+
+			string Command = $"-win 0 {UserInfoComplete} -cmd \"annotate -i \"{ DTE.ActiveDocument.FullName}\"\"";
+
+			TryP4VCommand(Command);
+		}
+
+		private void ChangeDiffSetting()
+		{
+			if (UnrealVSPackage.Instance.OptionsPage.AllowUnrealVSOverrideDiffSettings)
+			{
+				bool Margin = true;
+				UnrealVSPackage.Instance.EditorOptionsFactory.GlobalOptions.SetOptionValue("Diff/View/ShowDiffOverviewMargin", Margin);
+
+				DifferenceHighlightMode HighlightMode = (DifferenceHighlightMode)3;
+				UnrealVSPackage.Instance.EditorOptionsFactory.GlobalOptions.SetOptionValue("Diff/View/HighlightMode", HighlightMode);
+			}
+		}
+
+		private void P4DiffinVSHandler(object Sender, EventArgs Args)
+		{
+			DTE DTE = UnrealVSPackage.Instance.DTE;
+
+			if (DTE.ActiveDocument == null)
+			{
+				Logging.WriteLine("P4DiffInVSHandler called without an active document");
+
+				P4OutputPane.OutputString($"1>------ P4DiffInVSHandler called without an active document{Environment.NewLine}");
+
+				return;
+			}
+
+			ChangeDiffSetting();
+
+			// get the HAVE revision
+			// p4 fstat -T "haveRev" -Olp //UE5/Main/Engine/Source/Programs/UnrealVS/UnrealVS.Shared/P4Commands.cs
+			TryP4Command($"fstat -T \"haveRev\" -Olp {DTE.ActiveDocument.FullName}");
+
+			// expect output of the form "... haveRev 5"
+			Regex HavePattern = new Regex(@"... haveRev (?<Have>.*)$", RegexOptions.Compiled | RegexOptions.Multiline);
+			System.Text.RegularExpressions.Match HaveMatch = HavePattern.Match(P4OperationStdOut);
+			int HaveRev = Int32.Parse(HaveMatch.Groups["Have"].Value.Trim());
+
+			// Generate the Temp filename
+			string TempPath = Path.GetTempPath();
+			string TempFileName = Path.GetFileNameWithoutExtension(DTE.ActiveDocument.FullName) + "$" + HaveRev.ToString() + Path.GetExtension(DTE.ActiveDocument.FullName);
+			string TempFilePath = Path.Combine(TempPath, TempFileName);
+
+			// sync the HAVE revision to a file
+			// p4 print //UE5/Main/Engine/Source/Programs/UnrealVS/UnrealVS.Shared/P4Commands.cs#5 >> file
+
+			TryP4Command($"print {DTE.ActiveDocument.FullName}#{HaveRev}");
+
+			File.WriteAllText(TempFilePath, P4OperationStdOut);
 
+			// Tools.DiffFiles "file1.cs" "file2.cs"
+			DTE.ExecuteCommand("Tools.DiffFiles", $"\"{TempFilePath}\" \"{DTE.ActiveDocument.FullName}\"");
+
+		}
 		private void OpenForEdit(string FileName)
 		{
 			// Don't open for edit if the file is already writable
@@ -440,14 +577,38 @@ namespace UnrealVS
 			return WorkingDirectory;
 		}
 
+		void SetUserInfoStrings()
+		{
+			TryP4Command($"-s -L \"{P4WorkingDirectory}\" info", PullWorkingDirectoryOff);
+
+			Regex UserPattern = new Regex(@"User name: (?<user>.*)$", RegexOptions.Compiled | RegexOptions.Multiline);
+			Regex PortPattern = new Regex(@"Server address: (?<port>.*)$", RegexOptions.Compiled | RegexOptions.Multiline);
+			Regex ClientPattern = new Regex(@"Client name: (?<client>.*)$", RegexOptions.Compiled | RegexOptions.Multiline);
+
+			System.Text.RegularExpressions.Match UserMatch = UserPattern.Match(P4OperationStdOut);
+			System.Text.RegularExpressions.Match PortMatch = PortPattern.Match(P4OperationStdOut);
+			System.Text.RegularExpressions.Match ClientMath = ClientPattern.Match(P4OperationStdOut);
+
+			Port = PortMatch.Groups["port"].Value.Trim();
+			Username = UserMatch.Groups["user"].Value.Trim();
+			Client = ClientMath.Groups["client"].Value.Trim();
+
+			UserInfoComplete = string.Format(" -p {0} -u {1} -c {2} ", Port, Username, Client);
+
+			P4OutputPane.OutputString("GetUserInfoStringFull : " + UserInfoComplete);
+		}
+
 		private void PullWorkingDirectory(bool pullfromP4Settings)
 		{
-			if (IsSolutionLoaded() || P4WorkingDirectory != null || P4WorkingDirectory.Length < 2)
+			if (IsSolutionLoaded() && ( P4WorkingDirectory == null || P4WorkingDirectory.Length < 2 ))
 			{
 				DTE DTE = UnrealVSPackage.Instance.DTE;
 
 				// use the current solution folder as a temp working directory
 				P4WorkingDirectory = Path.GetDirectoryName(DTE.Solution.FileName);
+
+				// SetUserInfoStrings 
+				SetUserInfoStrings();
 			}
 
 			// if the callee wants us to pull a CWD and it wasn't already done
@@ -473,6 +634,26 @@ namespace UnrealVS
 		}
 
 		private bool TryP4Command(string CommandLine, bool PullWorkingDirectoryNow = PullWorkingDirectoryOn, bool OutputStdOut = OutputStdOutOn)
+		{
+			return TryP4CommandEx(P4Exe,CommandLine, PullWorkingDirectoryNow, OutputStdOut);
+		}
+
+		private bool TryP4VCCommand(string CommandLine, bool PullWorkingDirectoryNow = PullWorkingDirectoryOn, bool OutputStdOut = OutputStdOutOn)
+		{
+			if (P4VCCmd.Length > 1)
+			{
+				return TryP4CommandEx(P4VCCmd, CommandLine, PullWorkingDirectoryNow, OutputStdOut);
+			}
+
+			return false;
+		}
+
+		private bool TryP4VCommand(string CommandLine, bool PullWorkingDirectoryNow = PullWorkingDirectoryOn, bool OutputStdOut = OutputStdOutOn)
+		{
+			return TryP4CommandEx(P4VExe, CommandLine, PullWorkingDirectoryNow, OutputStdOut);
+		}
+
+		private bool TryP4CommandEx(string CmdPath, string CommandLine, bool PullWorkingDirectoryNow = PullWorkingDirectoryOn, bool OutputStdOut = OutputStdOutOn)
 		{
 			PullWorkingDirectory(PullWorkingDirectoryNow);
 		
@@ -508,7 +689,7 @@ namespace UnrealVS
 			{
 				StartInfo = new ProcessStartInfo()
 				{
-					FileName = P4Exe,
+					FileName = CmdPath,
 					Arguments = CommandLine,
 					WorkingDirectory = P4WorkingDirectory,
 					UseShellExecute = false,
@@ -529,7 +710,7 @@ namespace UnrealVS
 			P4OperationStdOut = StdOutSB.ToString();
 			P4OperationStdErr = StdErrSB.ToString();
 
-if (P4OperationStdOut.Length > 0 && OutputStdOut)
+			if (P4OperationStdOut.Length > 0 && OutputStdOut)
 			{
 				P4OutputPane.OutputString(P4OperationStdOut);
 			}
diff --git a/Engine/Source/Programs/UnrealVS/UnrealVS.Shared/UnrealVSOptions.cs b/Engine/Source/Programs/UnrealVS/UnrealVS.Shared/UnrealVSOptions.cs
--- a/Engine/Source/Programs/UnrealVS/UnrealVS.Shared/UnrealVSOptions.cs
+++ b/Engine/Source/Programs/UnrealVS/UnrealVS.Shared/UnrealVSOptions.cs
@@ -23,11 +23,8 @@ namespace UnrealVS
 			set { _HideNonGameStartupProjects = value; }
 		}
 
-		private bool _AllowUnrealVSP4;
-		private bool _AllowUnrealVSCheckoutOnEdit;
-
-		[Category("General")]
-		[DisplayName("Allow UnrealVS to auto checkout files on save")]
+		[Category("Unreal.P4")]
+		[DisplayName("Enable auto checkout on save")]
 		[Description("Uses p4 ini / environment settings to automatically checkout files on save, use 'RunUAT P4WriteConfig' to initialize those settings")]
 		public bool AllowUnrealVSCheckoutOnEdit
 		{
@@ -35,14 +32,27 @@ namespace UnrealVS
 			set { _AllowUnrealVSCheckoutOnEdit = value; }
 		}
 
-		[Category("General")]
-		[DisplayName("Allow UnrealVS to perform perforce operations")]
+		private bool _AllowUnrealVSCheckoutOnEdit;
+
+		[Category("Unreal.P4")]
+		[DisplayName("Override VS compare options")]
+		[Description("Unreal VS will override built in diff settings to the ideal for code, does not alter P4")]
+		public bool AllowUnrealVSOverrideDiffSettings
+		{
+			get { return _AllowUnrealVSOverrideDiffSettings; }
+			set { _AllowUnrealVSOverrideDiffSettings = value; }
+		}
+		private bool _AllowUnrealVSOverrideDiffSettings;
+
+		[Category("Unreal.P4")]
+		[DisplayName("Allow perforce operations")]
 		[Description("Uses p4 ini / environment settings to call P4 functionlality, use 'RunUAT P4WriteConfig' to initialize those settings")]
 		public bool AllowUnrealVSP4
 		{
 			get { return _AllowUnrealVSP4; }
 			set { _AllowUnrealVSP4 = value; }
 		}
+		private bool _AllowUnrealVSP4;
 
 		protected override void OnApply(PageApplyEventArgs e)
 		{
@@ -54,4 +64,6 @@ namespace UnrealVS
 			}
 		}
 	}
+
+
 }
diff --git a/Engine/Source/Programs/UnrealVS/UnrealVS.Shared/UnrealVSPackage.cs b/Engine/Source/Programs/UnrealVS/UnrealVS.Shared/UnrealVSPackage.cs
--- a/Engine/Source/Programs/UnrealVS/UnrealVS.Shared/UnrealVSPackage.cs
+++ b/Engine/Source/Programs/UnrealVS/UnrealVS.Shared/UnrealVSPackage.cs
@@ -6,6 +6,7 @@ using Microsoft.VisualStudio;
 using Microsoft.VisualStudio.OLE.Interop;
 using Microsoft.VisualStudio.Shell;
 using Microsoft.VisualStudio.Shell.Interop;
+using Microsoft.VisualStudio.Editor;
 using System;
 using System.Collections.Generic;
 using System.ComponentModel.Design;
@@ -14,7 +15,24 @@ using System.IO;
 using System.Linq;
 using System.Runtime.InteropServices;
 using System.Threading;
+using Microsoft.VisualStudio.TextManager.Interop;
 using Thread = System.Threading.Thread;
+using System;
+using System.Runtime.InteropServices;
+using System.Windows.Forms;
+using Microsoft.VisualStudio.ComponentModelHost;
+using Microsoft.VisualStudio.Editor;
+using Microsoft.VisualStudio.OLE.Interop;
+using Microsoft.VisualStudio.Shell.Interop;
+using Microsoft.VisualStudio.Shell;
+using Microsoft.VisualStudio.Text;
+using Microsoft.VisualStudio.Text.Editor;
+using Microsoft.VisualStudio.TextManager.Interop;
+using IOleServiceProvider = Microsoft.VisualStudio.OLE.Interop.IServiceProvider;
+using Microsoft.VisualStudio.Utilities;
+using Microsoft.VisualStudio.Text.Projection;
+using Microsoft.VisualStudio;
+using System.Collections.Generic;
 
 namespace UnrealVS
 {
@@ -201,6 +219,11 @@ namespace UnrealVS
 			DTE = await GetServiceAsync(typeof(DTE)) as DTE;
 			Logging.WriteLine("DTE version " + DTE.Version);
 
+			//TextManager = await GetServiceAsync(typeof(VsTextManagerClass)) as IVsTextManager3;
+
+			var componentModel = (IComponentModel)GetGlobalService(typeof(SComponentModel));
+			EditorOptionsFactory = componentModel.GetService<IEditorOptionsFactoryService>();
+
 			// Get selection manager and register to receive events
 			SelectionManager =
 				await GetServiceAsync(typeof(SVsShellMonitorSelection)) as IVsMonitorSelection;
@@ -389,6 +412,18 @@ namespace UnrealVS
 			}
 		}
 
+		//public IVsTextManager3 TextManager
+		//{
+		//	get;
+		//	private set;
+		//}
+
+		public IEditorOptionsFactoryService EditorOptionsFactory
+		{
+			get;
+			private set;
+		}
+
 		/// The package's options page
 		public UnrealVsOptions OptionsPage
 		{
