[UnrealVS] - add further hooks for operations that require checkout before save.
cleanup a few code guidelines problems
fix Diff titles
- checking this in without a VSIX drop.

diff --git a/Engine/Source/Programs/UnrealVS/UnrealVS.Shared/P4Commands.cs b/Engine/Source/Programs/UnrealVS/UnrealVS.Shared/P4Commands.cs
--- a/Engine/Source/Programs/UnrealVS/UnrealVS.Shared/P4Commands.cs
+++ b/Engine/Source/Programs/UnrealVS/UnrealVS.Shared/P4Commands.cs
@@ -3,7 +3,7 @@
 using EnvDTE;
 using Microsoft.VisualStudio.Shell;
 using Microsoft.VisualStudio.Shell.Interop;
-using Microsoft.VisualStudio.TextManager.Interop;
+using Microsoft.VisualStudio.Text.Differencing;
 using System;
 using System.Collections.Generic;
 using System.ComponentModel.Design;
@@ -11,10 +11,6 @@ using System.Diagnostics;
 using System.IO;
 using System.Text;
 using System.Text.RegularExpressions;
-using Microsoft.VisualStudio.Settings;
-using Microsoft.VisualStudio.Shell.Settings;
-using Microsoft.VisualStudio.Text.Editor;
-using Microsoft.VisualStudio.Text.Differencing;
 
 namespace UnrealVS
 {
@@ -23,7 +19,7 @@ namespace UnrealVS
 		private const bool PullWorkingDirectoryOn = true;
 		private const bool PullWorkingDirectoryOff = false;
 
-		private const bool OutputStdOutOn= true;
+		private const bool OutputStdOutOn = true;
 		private const bool OutputStdOutOff = false;
 
 		private const int P4SubMenuID = 0x3100;
@@ -32,7 +28,7 @@ namespace UnrealVS
 		private const int P4ViewSelectedCLButtonID = 0x1452;
 		private const int P4IntegrationAwareTimelapseButtonID = 0x1453;
 		private const int P4DiffinVSButtonID = 0x1454;
-		
+
 
 
 		private OleMenuCommand SubMenuCommand;
@@ -108,7 +104,7 @@ namespace UnrealVS
 
 			// create specific output window for unrealvs.P4
 			P4OutputPane = UnrealVSPackage.Instance.GetP4OutputPane();
-			
+
 			// figure out the P4VC path
 			if (!File.Exists(P4VCCmd))
 			{
@@ -116,13 +112,13 @@ namespace UnrealVS
 				{
 					P4VCCmd = P4VCCmdBat;
 				}
-				else 
+				else
 				{
 					P4OutputPane.Activate();
 					P4OutputPane.OutputString($"1>------ P4VC not found, {P4VCCmd} or {P4VCCmd}{Environment.NewLine}");
 					P4VCCmd = "";
 				}
-					
+
 			}
 
 			// add commands
@@ -132,7 +128,7 @@ namespace UnrealVS
 			P4CommandsList.Add(new P4Command(P4DiffinVSButtonID, P4DiffinVSHandler));
 
 			if (P4VCCmd.Length > 1)
-			{	
+			{
 				P4CommandsList.Add(new P4Command(P4ViewSelectedCLButtonID, P4ViewSelectedCLButtonHandler));
 
 			}
@@ -148,7 +144,17 @@ namespace UnrealVS
 			// hook up the on "Save" list - internally they verify the users choice
 			RegisterCallbackHandler("File.SaveSelectedItems", SaveSelectedCallback);
 			RegisterCallbackHandler("File.SaveAll", OnSaveAll);
-			
+
+			// Hook up "Dirty" operations - things like 'build' when files are edited but not checked out hit this path
+			RegisterCallbackHandler("Build.BuildSolution", SaveModifiedFiles);
+			RegisterCallbackHandler("Build.Compile", SaveModifiedFiles);
+			RegisterCallbackHandler("Build.BuildOnlyProject", SaveModifiedFiles);
+			RegisterCallbackHandler("Debug.Start", SaveModifiedFiles);
+
+			// there are *other* ways to hit the same operations - add those too.
+			RegisterCallbackHandler("ClassViewContextMenus.ClassViewProject.Build", SaveModifiedFiles);
+			RegisterCallbackHandler("ClassViewContextMenus.ClassViewProject.Rebuild", SaveModifiedFiles);
+			RegisterCallbackHandler("ClassViewContextMenus.ClassViewProject.Debug.Startnewinstance", SaveModifiedFiles);
 		}
 		// Called when solutions are loaded or unloaded
 		private void SoltuionOpened()
@@ -187,7 +193,7 @@ namespace UnrealVS
 
 				}
 			}
-		
+
 			if (Event != null)
 			{
 				Event.BeforeExecute += Callback;
@@ -195,6 +201,20 @@ namespace UnrealVS
 			}
 		}
 
+		private void SaveModifiedFiles(string Guid, int ID, object CustomIn, object CustomOut, ref bool CancelDefault)
+		{
+			if (UnrealVSPackage.Instance.OptionsPage.AllowUnrealVSCheckoutOnEdit)
+			{
+				foreach (Document File in UnrealVSPackage.Instance.DTE.Documents)
+				{
+					if (!File.Saved && File.ReadOnly)
+					{
+						OpenForEdit(File.FullName);
+					}
+				}
+			}
+		}
+
 		void SaveSelectedCallback(string Guid, int ID, object CustomIn, object CustomOut, ref bool CancelDefault)
 		{
 			if (UnrealVSPackage.Instance.OptionsPage.AllowUnrealVSCheckoutOnEdit)
@@ -256,7 +276,7 @@ namespace UnrealVS
 			{
 				OpenForEdit(Proj.FullName);
 			}
-				
+
 			if (Proj.ProjectItems == null)
 			{
 				return;
@@ -317,7 +337,7 @@ namespace UnrealVS
 		}
 
 		private void P4CheckoutButtonHandler(object Sender, EventArgs Args)
-        {
+		{
 			DTE DTE = UnrealVSPackage.Instance.DTE;
 
 			// Check we've got a file open
@@ -331,7 +351,7 @@ namespace UnrealVS
 				return;
 			}
 
-			OpenForEdit(DTE.ActiveDocument.FullName);	
+			OpenForEdit(DTE.ActiveDocument.FullName);
 		}
 		private void P4AnnotateButtonHandler(object Sender, EventArgs Args)
 		{
@@ -350,9 +370,9 @@ namespace UnrealVS
 			}
 
 			// Call annotate itself
-			bool Result= TryP4Command($"annotate -TcIqu \"{DTE.ActiveDocument.FullName}", PullWorkingDirectoryOn, OutputStdOutOff);
+			bool Result = TryP4Command($"annotate -TcIqu \"{DTE.ActiveDocument.FullName}", PullWorkingDirectoryOn, OutputStdOutOff);
 
-			if (!Result || P4OperationStdErr.Length>0)
+			if (!Result || P4OperationStdErr.Length > 0)
 			{
 				P4OutputPane.OutputString($"1>------ P4Annotate call failed");
 				return;
@@ -447,13 +467,13 @@ namespace UnrealVS
 			{
 				ChangeList = -2;
 			}
-			
+
 
 			if (ChangeList > 0)
 			{
 				TryP4VCCommand($"Change {ChangeList}");
 			}
-			
+
 		}
 
 		private void P4IntegrationAwareTimeLapseHandler(object Sender, EventArgs Args)
@@ -503,27 +523,35 @@ namespace UnrealVS
 
 			// get the HAVE revision
 			// p4 fstat -T "haveRev" -Olp //UE5/Main/Engine/Source/Programs/UnrealVS/UnrealVS.Shared/P4Commands.cs
-			TryP4Command($"fstat -T \"haveRev\" -Olp {DTE.ActiveDocument.FullName}");
-
-			// expect output of the form "... haveRev 5"
-			Regex HavePattern = new Regex(@"... haveRev (?<Have>.*)$", RegexOptions.Compiled | RegexOptions.Multiline);
-			System.Text.RegularExpressions.Match HaveMatch = HavePattern.Match(P4OperationStdOut);
-			int HaveRev = Int32.Parse(HaveMatch.Groups["Have"].Value.Trim());
-
-			// Generate the Temp filename
-			string TempPath = Path.GetTempPath();
-			string TempFileName = Path.GetFileNameWithoutExtension(DTE.ActiveDocument.FullName) + "$" + HaveRev.ToString() + Path.GetExtension(DTE.ActiveDocument.FullName);
-			string TempFilePath = Path.Combine(TempPath, TempFileName);
-
-			// sync the HAVE revision to a file
-			// p4 print //UE5/Main/Engine/Source/Programs/UnrealVS/UnrealVS.Shared/P4Commands.cs#5 >> file
-
-			TryP4Command($"print {DTE.ActiveDocument.FullName}#{HaveRev}");
-
-			File.WriteAllText(TempFilePath, P4OperationStdOut);
+			if (TryP4Command($"fstat -T \"haveRev,depotFile\" -Olp \"{DTE.ActiveDocument.FullName}\""))
+			{
+				// expect output of the form
+				//		"... haveRev 5"
+				//		"... depotFile //UE5/Main/Engine/Source/Programs/UnrealVS/UnrealVS.Shared/P4Commands.cs
+				Regex HavePattern = new Regex(@"... haveRev (?<Have>.*)$", RegexOptions.Compiled | RegexOptions.Multiline);
+				Regex DepotPathPattern = new Regex(@"... depotFile (?<depotFile>.*)$", RegexOptions.Compiled | RegexOptions.Multiline);
+				System.Text.RegularExpressions.Match HaveMatch = HavePattern.Match(P4OperationStdOut);
+				System.Text.RegularExpressions.Match PathMatch = DepotPathPattern.Match(P4OperationStdOut);
+				int HaveRev = Int32.Parse(HaveMatch.Groups["Have"].Value.Trim());
+				string depotPath = PathMatch.Groups["depotFile"].Value.Trim();
+
+				// Generate the Temp filename
+				string TempPath = Path.GetTempPath();
+				string TempFileName = Path.GetFileNameWithoutExtension(DTE.ActiveDocument.FullName) + "$" + HaveRev.ToString() + Path.GetExtension(DTE.ActiveDocument.FullName);
+				string TempFilePath = Path.Combine(TempPath, TempFileName);
+
+				// sync the HAVE revision to a file
+				// p4 print //UE5/Main/Engine/Source/Programs/UnrealVS/UnrealVS.Shared/P4Commands.cs#5 >> file
+
+				string VersionPath = $"{depotPath}#{HaveRev}";
+				if (TryP4Command($"-q print \"{VersionPath}\""))
+				{
+					File.WriteAllText(TempFilePath, P4OperationStdOut);
 
-			// Tools.DiffFiles "file1.cs" "file2.cs"
-			DTE.ExecuteCommand("Tools.DiffFiles", $"\"{TempFilePath}\" \"{DTE.ActiveDocument.FullName}\"");
+					// Tools.DiffFiles SourceFile, TargetFile, [SourceDisplayName],[TargetDisplayName]
+					DTE.ExecuteCommand("Tools.DiffFiles", $"\"{TempFilePath}\" \"{DTE.ActiveDocument.FullName}\" \"{VersionPath}\" \"{DTE.ActiveDocument.FullName}\"");
+				}
+			}
 
 		}
 		private void OpenForEdit(string FileName)
@@ -600,7 +628,7 @@ namespace UnrealVS
 
 		private void PullWorkingDirectory(bool pullfromP4Settings)
 		{
-			if (IsSolutionLoaded() && ( P4WorkingDirectory == null || P4WorkingDirectory.Length < 2 ))
+			if (IsSolutionLoaded() && (P4WorkingDirectory == null || P4WorkingDirectory.Length < 2))
 			{
 				DTE DTE = UnrealVSPackage.Instance.DTE;
 
@@ -629,13 +657,13 @@ namespace UnrealVS
 					P4OutputPane.OutputString($"P4WorkingDirectory set failed {Environment.NewLine}");
 				}
 
-				
+
 			}
 		}
 
 		private bool TryP4Command(string CommandLine, bool PullWorkingDirectoryNow = PullWorkingDirectoryOn, bool OutputStdOut = OutputStdOutOn)
 		{
-			return TryP4CommandEx(P4Exe,CommandLine, PullWorkingDirectoryNow, OutputStdOut);
+			return TryP4CommandEx(P4Exe, CommandLine, PullWorkingDirectoryNow, OutputStdOut);
 		}
 
 		private bool TryP4VCCommand(string CommandLine, bool PullWorkingDirectoryNow = PullWorkingDirectoryOn, bool OutputStdOut = OutputStdOutOn)
@@ -656,7 +684,7 @@ namespace UnrealVS
 		private bool TryP4CommandEx(string CmdPath, string CommandLine, bool PullWorkingDirectoryNow = PullWorkingDirectoryOn, bool OutputStdOut = OutputStdOutOn)
 		{
 			PullWorkingDirectory(PullWorkingDirectoryNow);
-		
+
 			DTE DTE = UnrealVSPackage.Instance.DTE;
 
 			// If there's already a operation in progress, abort... potential issues here
@@ -700,7 +728,7 @@ namespace UnrealVS
 			};
 			// Create a delegate for handling output messages
 			ChildProcess.OutputDataReceived += (s, a) => { if (a.Data != null) StdOutSB.AppendLine(a.Data); };
-			ChildProcess.ErrorDataReceived += (s, a) => { if (a.Data != null) StdErrSB.AppendLine(a.Data);	};
+			ChildProcess.ErrorDataReceived += (s, a) => { if (a.Data != null) StdErrSB.AppendLine(a.Data); };
 			ChildProcess.Start();
 			ChildProcess.BeginOutputReadLine();
 			ChildProcess.BeginErrorReadLine();
@@ -722,7 +750,7 @@ namespace UnrealVS
 
 			P4OutputPane.OutputString($"1>------ P4 Operation complete{Environment.NewLine}");
 
-			return true;
+			return P4OperationStdErr.Length == 0;
 		}
 	}
 }
diff --git a/Engine/Source/Programs/UnrealVS/UnrealVS.Shared/UnrealVSOptions.cs b/Engine/Source/Programs/UnrealVS/UnrealVS.Shared/UnrealVSOptions.cs
--- a/Engine/Source/Programs/UnrealVS/UnrealVS.Shared/UnrealVSOptions.cs
+++ b/Engine/Source/Programs/UnrealVS/UnrealVS.Shared/UnrealVSOptions.cs
@@ -12,47 +12,25 @@ namespace UnrealVS
 	{
 		public event EventHandler OnOptionsChanged;
 
-		private bool _HideNonGameStartupProjects;
-
 		[Category("General")]
 		[DisplayName("Hide Non-Game Startup Projects")]
 		[Description("Shows only game projects in the startup project and batch-builder lists")]
-		public bool HideNonGameStartupProjects
-		{
-			get { return _HideNonGameStartupProjects; }
-			set { _HideNonGameStartupProjects = value; }
-		}
+		public bool HideNonGameStartupProjects { get; set; }
 
 		[Category("Unreal.P4")]
 		[DisplayName("Enable auto checkout on save")]
 		[Description("Uses p4 ini / environment settings to automatically checkout files on save, use 'RunUAT P4WriteConfig' to initialize those settings")]
-		public bool AllowUnrealVSCheckoutOnEdit
-		{
-			get { return _AllowUnrealVSCheckoutOnEdit; }
-			set { _AllowUnrealVSCheckoutOnEdit = value; }
-		}
-
-		private bool _AllowUnrealVSCheckoutOnEdit;
+		public bool AllowUnrealVSCheckoutOnEdit { get; set; }
 
 		[Category("Unreal.P4")]
 		[DisplayName("Override VS compare options")]
 		[Description("Unreal VS will override built in diff settings to the ideal for code, does not alter P4")]
-		public bool AllowUnrealVSOverrideDiffSettings
-		{
-			get { return _AllowUnrealVSOverrideDiffSettings; }
-			set { _AllowUnrealVSOverrideDiffSettings = value; }
-		}
-		private bool _AllowUnrealVSOverrideDiffSettings;
+		public bool AllowUnrealVSOverrideDiffSettings { get; set; }
 
 		[Category("Unreal.P4")]
 		[DisplayName("Allow perforce operations")]
 		[Description("Uses p4 ini / environment settings to call P4 functionlality, use 'RunUAT P4WriteConfig' to initialize those settings")]
-		public bool AllowUnrealVSP4
-		{
-			get { return _AllowUnrealVSP4; }
-			set { _AllowUnrealVSP4 = value; }
-		}
-		private bool _AllowUnrealVSP4;
+		public bool AllowUnrealVSP4 { get; set; }
 
 		protected override void OnApply(PageApplyEventArgs e)
 		{
