[UnrealVS] - implement RunningDocumentTable and use onSave callback to cover all edit/save operations (Better than manually figuring out the right hooks to enter)

diff --git a/Engine/Source/Programs/UnrealVS/UnrealVS.2019/Properties/AssemblyInfo.cs b/Engine/Source/Programs/UnrealVS/UnrealVS.2019/Properties/AssemblyInfo.cs
--- a/Engine/Source/Programs/UnrealVS/UnrealVS.2019/Properties/AssemblyInfo.cs
+++ b/Engine/Source/Programs/UnrealVS/UnrealVS.2019/Properties/AssemblyInfo.cs
@@ -30,8 +30,8 @@ using System.Runtime.InteropServices;
 // You can specify all the values or you can default the Revision and Build Numbers 
 // by using the '*' as shown below:
 
-[assembly: AssemblyVersion("1.64.0.0")]
-[assembly: AssemblyFileVersion("1.64.0.0")]
+[assembly: AssemblyVersion("1.65.0.0")]
+[assembly: AssemblyFileVersion("1.65.0.0")]
 
 
 
diff --git a/Engine/Source/Programs/UnrealVS/UnrealVS.2022/Properties/AssemblyInfo.cs b/Engine/Source/Programs/UnrealVS/UnrealVS.2022/Properties/AssemblyInfo.cs
--- a/Engine/Source/Programs/UnrealVS/UnrealVS.2022/Properties/AssemblyInfo.cs
+++ b/Engine/Source/Programs/UnrealVS/UnrealVS.2022/Properties/AssemblyInfo.cs
@@ -30,8 +30,8 @@ using System.Runtime.InteropServices;
 // You can specify all the values or you can default the Revision and Build Numbers 
 // by using the '*' as shown below:
 
-[assembly: AssemblyVersion("1.64.0.0")]
-[assembly: AssemblyFileVersion("1.64.0.0")]
+[assembly: AssemblyVersion("1.65.0.0")]
+[assembly: AssemblyFileVersion("1.65.0.0")]
 
 
 
diff --git a/Engine/Source/Programs/UnrealVS/UnrealVS.Shared/P4Commands.cs b/Engine/Source/Programs/UnrealVS/UnrealVS.Shared/P4Commands.cs
--- a/Engine/Source/Programs/UnrealVS/UnrealVS.Shared/P4Commands.cs
+++ b/Engine/Source/Programs/UnrealVS/UnrealVS.Shared/P4Commands.cs
@@ -1,6 +1,7 @@
 // Copyright Epic Games, Inc. All Rights Reserved.
 
 using EnvDTE;
+using Microsoft.VisualStudio;
 using Microsoft.VisualStudio.Shell;
 using Microsoft.VisualStudio.Shell.Interop;
 using Microsoft.VisualStudio.Text.Differencing;
@@ -29,8 +30,6 @@ namespace UnrealVS
 		private const int P4IntegrationAwareTimelapseButtonID = 0x1453;
 		private const int P4DiffinVSButtonID = 0x1454;
 
-
-
 		private OleMenuCommand SubMenuCommand;
 
 		private System.Diagnostics.Process ChildProcess;
@@ -76,6 +75,8 @@ namespace UnrealVS
 
 		private List<P4Command> P4CommandsList = new List<P4Command>();
 
+		private IntercepteSave Interceptor;
+
 		private bool IsSolutionLoaded()
 		{
 			DTE DTE = UnrealVSPackage.Instance.DTE;
@@ -141,20 +142,11 @@ namespace UnrealVS
 			// Update the menu visibility to enforce user options.
 			UpdateMenuOptions();
 
-			// hook up the on "Save" list - internally they verify the users choice
-			RegisterCallbackHandler("File.SaveSelectedItems", SaveSelectedCallback);
-			RegisterCallbackHandler("File.SaveAll", OnSaveAll);
+			var runningDocumentTable = new RunningDocumentTable(UnrealVSPackage.Instance);
+			Interceptor = new IntercepteSave(UnrealVSPackage.Instance.DTE, runningDocumentTable, this);
 
-			// Hook up "Dirty" operations - things like 'build' when files are edited but not checked out hit this path
-			RegisterCallbackHandler("Build.BuildSolution", SaveModifiedFiles);
-			RegisterCallbackHandler("Build.Compile", SaveModifiedFiles);
-			RegisterCallbackHandler("Build.BuildOnlyProject", SaveModifiedFiles);
-			RegisterCallbackHandler("Debug.Start", SaveModifiedFiles);
+			runningDocumentTable.Advise(Interceptor);
 
-			// there are *other* ways to hit the same operations - add those too.
-			RegisterCallbackHandler("ClassViewContextMenus.ClassViewProject.Build", SaveModifiedFiles);
-			RegisterCallbackHandler("ClassViewContextMenus.ClassViewProject.Rebuild", SaveModifiedFiles);
-			RegisterCallbackHandler("ClassViewContextMenus.ClassViewProject.Debug.Startnewinstance", SaveModifiedFiles);
 		}
 		// Called when solutions are loaded or unloaded
 		private void SoltuionOpened()
@@ -355,6 +347,8 @@ namespace UnrealVS
 		}
 		private void P4AnnotateButtonHandler(object Sender, EventArgs Args)
 		{
+			// Debug.ListCallStack
+
 			DTE DTE = UnrealVSPackage.Instance.DTE;
 
 			P4OutputPane.Activate();
@@ -554,7 +548,7 @@ namespace UnrealVS
 			}
 
 		}
-		private void OpenForEdit(string FileName)
+		public void OpenForEdit(string FileName)
 		{
 			// Don't open for edit if the file is already writable
 			if (!File.Exists(FileName) || !File.GetAttributes(FileName).HasFlag(FileAttributes.ReadOnly))
@@ -753,4 +747,74 @@ namespace UnrealVS
 			return P4OperationStdErr.Length == 0;
 		}
 	}
+
+	// 
+
+	internal class IntercepteSave : IVsRunningDocTableEvents3
+	{
+		private readonly DTE DTE;
+		private readonly RunningDocumentTable RunningDocumentTable;
+		private readonly P4Commands P4Ops;
+
+		public IntercepteSave(DTE InDTE, RunningDocumentTable InRrunningDocumentTable, P4Commands InP4Ops)
+		{
+			DTE = InDTE;
+			RunningDocumentTable = InRrunningDocumentTable;
+			P4Ops = InP4Ops;
+		}
+
+		public int OnBeforeSave(uint DocumentCookie)
+		{
+			RunningDocumentInfo DocumentInfo = RunningDocumentTable.GetDocumentInfo(DocumentCookie);
+
+			string AbsoluteFilePath = DocumentInfo.Moniker;
+
+			P4Ops.OpenForEdit(AbsoluteFilePath);
+
+			return VSConstants.S_OK;
+		}
+
+		// we are only using OnBeforeSave - but the interface requires us to define the whole interface.
+		public int OnAfterFirstDocumentLock(uint docCookie, uint dwRdtLockType, uint dwReadLocksRemaining, uint dwEditLocksRemaining)
+		{
+			return VSConstants.S_OK;
+		}
+
+		public int OnBeforeLastDocumentUnlock(uint docCookie, uint dwRdtLockType, uint dwReadLocksRemaining, uint dwEditLocksRemaining)
+		{
+			return VSConstants.S_OK;
+		}
+
+		public int OnAfterSave(uint docCookie)
+		{
+			return VSConstants.S_OK;
+		}
+
+		public int OnAfterAttributeChange(uint docCookie, uint grfAttribs)
+		{
+			return VSConstants.S_OK;
+		}
+
+		public int OnBeforeDocumentWindowShow(uint docCookie, int fFirstShow, IVsWindowFrame pFrame)
+		{
+			return VSConstants.S_OK;
+		}
+
+		public int OnAfterDocumentWindowHide(uint docCookie, IVsWindowFrame pFrame)
+		{
+			return VSConstants.S_OK;
+		}
+
+		int IVsRunningDocTableEvents3.OnAfterAttributeChangeEx(uint docCookie, uint grfAttribs, IVsHierarchy pHierOld, uint itemidOld,
+			string pszMkDocumentOld, IVsHierarchy pHierNew, uint itemidNew, string pszMkDocumentNew)
+		{
+			return VSConstants.S_OK;
+		}
+
+		int IVsRunningDocTableEvents2.OnAfterAttributeChangeEx(uint docCookie, uint grfAttribs, IVsHierarchy pHierOld, uint itemidOld,
+			string pszMkDocumentOld, IVsHierarchy pHierNew, uint itemidNew, string pszMkDocumentNew)
+		{
+			return VSConstants.S_OK;
+		}
+	}
 }
