[UnrealVS] - add P4ShowfileInP4VS & P4GetLast10Changes commands

diff --git a/Engine/Source/Programs/UnrealVS/UnrealVS.2019/Properties/AssemblyInfo.cs b/Engine/Source/Programs/UnrealVS/UnrealVS.2019/Properties/AssemblyInfo.cs
--- a/Engine/Source/Programs/UnrealVS/UnrealVS.2019/Properties/AssemblyInfo.cs
+++ b/Engine/Source/Programs/UnrealVS/UnrealVS.2019/Properties/AssemblyInfo.cs
@@ -30,8 +30,8 @@ using System.Runtime.InteropServices;
 // You can specify all the values or you can default the Revision and Build Numbers 
 // by using the '*' as shown below:
 
-[assembly: AssemblyVersion("1.65.0.0")]
-[assembly: AssemblyFileVersion("1.65.0.0")]
+[assembly: AssemblyVersion("1.66.0.0")]
+[assembly: AssemblyFileVersion("1.66.0.0")]
 
 
 
diff --git a/Engine/Source/Programs/UnrealVS/UnrealVS.2022/Properties/AssemblyInfo.cs b/Engine/Source/Programs/UnrealVS/UnrealVS.2022/Properties/AssemblyInfo.cs
--- a/Engine/Source/Programs/UnrealVS/UnrealVS.2022/Properties/AssemblyInfo.cs
+++ b/Engine/Source/Programs/UnrealVS/UnrealVS.2022/Properties/AssemblyInfo.cs
@@ -30,8 +30,8 @@ using System.Runtime.InteropServices;
 // You can specify all the values or you can default the Revision and Build Numbers 
 // by using the '*' as shown below:
 
-[assembly: AssemblyVersion("1.65.0.0")]
-[assembly: AssemblyFileVersion("1.65.0.0")]
+[assembly: AssemblyVersion("1.66.0.0")]
+[assembly: AssemblyFileVersion("1.66.0.0")]
 
 
 
diff --git a/Engine/Source/Programs/UnrealVS/UnrealVS.Shared/P4Commands.cs b/Engine/Source/Programs/UnrealVS/UnrealVS.Shared/P4Commands.cs
--- a/Engine/Source/Programs/UnrealVS/UnrealVS.Shared/P4Commands.cs
+++ b/Engine/Source/Programs/UnrealVS/UnrealVS.Shared/P4Commands.cs
@@ -29,6 +29,8 @@ namespace UnrealVS
 		private const int P4ViewSelectedCLButtonID = 0x1452;
 		private const int P4IntegrationAwareTimelapseButtonID = 0x1453;
 		private const int P4DiffinVSButtonID = 0x1454;
+		private const int P4GetLast10ChangesID = 0x1455;
+		private const int P4ShowFileinP4VID = 0x1456;
 
 		private OleMenuCommand SubMenuCommand;
 
@@ -127,6 +129,8 @@ namespace UnrealVS
 			P4CommandsList.Add(new P4Command(P4AnnotateButtonID, P4AnnotateButtonHandler));
 			P4CommandsList.Add(new P4Command(P4IntegrationAwareTimelapseButtonID, P4IntegrationAwareTimeLapseHandler));
 			P4CommandsList.Add(new P4Command(P4DiffinVSButtonID, P4DiffinVSHandler));
+			P4CommandsList.Add(new P4Command(P4GetLast10ChangesID, P4GetLast10ChangesHandler));
+			P4CommandsList.Add(new P4Command(P4ShowFileinP4VID, P4ShowFileInP4Vandler));
 
 			if (P4VCCmd.Length > 1)
 			{
@@ -538,7 +542,7 @@ namespace UnrealVS
 				// p4 print //UE5/Main/Engine/Source/Programs/UnrealVS/UnrealVS.Shared/P4Commands.cs#5 >> file
 
 				string VersionPath = $"{depotPath}#{HaveRev}";
-				if (TryP4Command($"-q print \"{VersionPath}\""))
+				if (TryP4Command($"-q print \"{VersionPath}\"",PullWorkingDirectoryOn, OutputStdOutOff))
 				{
 					File.WriteAllText(TempFilePath, P4OperationStdOut);
 
@@ -546,7 +550,39 @@ namespace UnrealVS
 					DTE.ExecuteCommand("Tools.DiffFiles", $"\"{TempFilePath}\" \"{DTE.ActiveDocument.FullName}\" \"{VersionPath}\" \"{DTE.ActiveDocument.FullName}\"");
 				}
 			}
+		}
+
+		private void P4GetLast10ChangesHandler(object Sender, EventArgs Args)
+		{
+			DTE DTE = UnrealVSPackage.Instance.DTE;
+
+			string TempPath = Path.GetTempPath();
 
+			// generate the changes list
+			string ChangesTempFileName = Path.GetFileNameWithoutExtension(DTE.ActiveDocument.FullName) + "_last_10_changelists" + Path.GetExtension(DTE.ActiveDocument.FullName);
+			string ChangesTempFilePath = Path.Combine(TempPath, ChangesTempFileName);
+
+
+			if (TryP4Command($"changes -itls submitted -m 10 {DTE.ActiveDocument.FullName}"))
+			{
+				File.WriteAllText(ChangesTempFilePath, P4OperationStdOut);
+
+				DTE.ExecuteCommand("File.OpenFile", $"\"{ChangesTempFilePath}\"");
+				DTE.ActiveDocument.Activate();
+				DTE.ExecuteCommand("Window.NewVerticalTabGroup");
+			}
+		}
+
+		private void P4ShowFileInP4Vandler(object Sender, EventArgs Args)
+		{
+			DTE DTE = UnrealVSPackage.Instance.DTE;
+
+			if (DTE.ActiveDocument == null)
+			{
+				return;
+			}
+
+			TryP4VCommand($"-s {DTE.ActiveDocument.FullName}");
 		}
 		public void OpenForEdit(string FileName)
 		{
@@ -685,9 +721,9 @@ namespace UnrealVS
 			// in the OnSave/OnSaveAll operations.
 			if (ChildProcess != null && !ChildProcess.HasExited)
 			{
-				P4OutputPane.OutputString($"P4 operation already in flight.{Environment.NewLine}");
+				P4OutputPane.OutputString($"P4 operation already in flight...waiting{Environment.NewLine}");
 				P4OutputPane.Activate();
-				return false;
+				ChildProcess.WaitForExit();
 			}
 
 			// Make sure any existing op is stopped
@@ -699,10 +735,6 @@ namespace UnrealVS
 			P4OutputPane.OutputString($"CWD: {P4WorkingDirectory}{Environment.NewLine}");
 			P4OutputPane.OutputString($"CMD: {P4Exe} {CommandLine}{Environment.NewLine}");
 
-			// Create a delegate for handling output messages
-			//void OutputHandler(object Sender, DataReceivedEventArgs Args) { if (Args.Data != null) P4OperationStdOut += $"{Args.Data}{Environment.NewLine}"; }
-			//void ErrOutputHandler(object Sender, DataReceivedEventArgs Args) { if (Args.Data != null) P4OperationStdErr += $"{Args.Data}{Environment.NewLine}"; }
-
 			StringBuilder StdOutSB = new StringBuilder();
 			StringBuilder StdErrSB = new StringBuilder();
 
