[UnrealVS] - cleanup P4Commands, apply some coding guidelines

diff --git a/Engine/Source/Programs/UnrealVS/UnrealVS.Shared/P4Commands.cs b/Engine/Source/Programs/UnrealVS/UnrealVS.Shared/P4Commands.cs
--- a/Engine/Source/Programs/UnrealVS/UnrealVS.Shared/P4Commands.cs
+++ b/Engine/Source/Programs/UnrealVS/UnrealVS.Shared/P4Commands.cs
@@ -17,11 +17,11 @@ namespace UnrealVS
 {
 	internal class P4Commands : IDisposable
 	{
-		private const bool PullWorkingDirectoryOn = true;
-		private const bool PullWorkingDirectoryOff = false;
+		private const bool bPullWorkingDirectoryOn = true;
+		private const bool bPullWorkingDirectoryOff = false;
 
-		private const bool OutputStdOutOn = true;
-		private const bool OutputStdOutOff = false;
+		private const bool bOutputStdOutOn = true;
+		private const bool bOutputStdOutOff = false;
 
 		private const int P4SubMenuID = 0x3100;
 		private const int P4CheckoutButtonID = 0x1450;
@@ -37,7 +37,7 @@ namespace UnrealVS
 		private System.Diagnostics.Process ChildProcess;
 		private IVsOutputWindowPane P4OutputPane;
 		private string P4WorkingDirectory;
-		private bool PullWorkingDirectorFromP4 = true;
+		private bool bPullWorkingDirectorFromP4 = true;
 
 		// stdXX from last operation
 		private string P4OperationStdOut;
@@ -90,12 +90,12 @@ namespace UnrealVS
 		{
 			DTE DTE = UnrealVSPackage.Instance.DTE;
 
-			bool EnableCommands = UnrealVSPackage.Instance.OptionsPage.AllowUnrealVSP4;
-			bool SolutionLoaded = IsSolutionLoaded();
+			bool bEnableCommands = UnrealVSPackage.Instance.OptionsPage.AllowUnrealVSP4;
+			bool bSolutionLoaded = IsSolutionLoaded();
 
 			var SenderSubMenuCommand = (OleMenuCommand)sender;
 
-			SenderSubMenuCommand.Visible = SenderSubMenuCommand.Enabled = EnableCommands & SolutionLoaded;
+			SenderSubMenuCommand.Visible = SenderSubMenuCommand.Enabled = bEnableCommands & bSolutionLoaded;
 		}
 
 		public P4Commands()
@@ -166,7 +166,7 @@ namespace UnrealVS
 
 			// Clear any existing P4 working directory settings
 			P4WorkingDirectory = "";
-			PullWorkingDirectorFromP4 = true;
+			bPullWorkingDirectorFromP4 = true;
 		}
 
 		void RegisterCallbackHandler(string CommandName, _dispCommandEvents_BeforeExecuteEventHandler Callback)
@@ -197,75 +197,6 @@ namespace UnrealVS
 			}
 		}
 
-		private void SaveModifiedFiles(string Guid, int ID, object CustomIn, object CustomOut, ref bool CancelDefault)
-		{
-			if (UnrealVSPackage.Instance.OptionsPage.AllowUnrealVSCheckoutOnEdit)
-			{
-				foreach (Document File in UnrealVSPackage.Instance.DTE.Documents)
-				{
-					if (!File.Saved && File.ReadOnly)
-					{
-						OpenForEdit(File.FullName);
-					}
-				}
-			}
-		}
-
-		void SaveSelectedCallback(string Guid, int ID, object CustomIn, object CustomOut, ref bool CancelDefault)
-		{
-			if (UnrealVSPackage.Instance.OptionsPage.AllowUnrealVSCheckoutOnEdit)
-			{
-				DTE DTE = UnrealVSPackage.Instance.DTE;
-
-				foreach (SelectedItem Item in DTE.SelectedItems)
-				{
-					if (Item.Project != null)
-					{
-						OpenForEdit(Item.Project.FullName);
-					}
-					else if (Item.ProjectItem != null)
-					{
-						OpenForEdit(Item.ProjectItem.Document.FullName);
-					}
-					else
-					{
-						OpenForEdit(DTE.Solution.FullName);
-					}
-				}
-			}
-		}
-
-		private void OnSaveAll(string Guid, int ID, object CustomIn, object CustomOut, ref bool CancelDefault)
-		{
-			if (UnrealVSPackage.Instance.OptionsPage.AllowUnrealVSCheckoutOnEdit)
-			{
-				DTE DTE = UnrealVSPackage.Instance.DTE;
-
-				if (!DTE.Solution.Saved)
-				{
-					OpenForEdit(DTE.Solution.FullName);
-				}
-
-				foreach (Document OpenedFile in DTE.Documents)
-				{
-					if (OpenedFile.Saved == false)
-					{
-						OpenForEdit(OpenedFile.FullName);
-					}
-				}
-
-				if (DTE.Solution.Projects == null)
-				{
-					return;
-				}
-
-				foreach (Project p in DTE.Solution.Projects)
-				{
-					EditProjectRecursive(p);
-				}
-			}
-		}
-
 		private void EditProjectRecursive(Project Proj)
 		{
 			if (!Proj.Saved)
@@ -298,13 +229,13 @@ namespace UnrealVS
 		{
 			DTE DTE = UnrealVSPackage.Instance.DTE;
 
-			bool EnableCommands = UnrealVSPackage.Instance.OptionsPage.AllowUnrealVSP4;
-			bool SolutionLoaded = IsSolutionLoaded();
+			bool bEnableCommands = UnrealVSPackage.Instance.OptionsPage.AllowUnrealVSP4;
+			bool bSolutionLoaded = IsSolutionLoaded();
 
 			// update each menu item enabled command
 			foreach (P4Command command in P4CommandsList)
 			{
-				command.Toggle(SolutionLoaded && EnableCommands);
+				command.Toggle(bSolutionLoaded && bEnableCommands);
 			}
 
 		}
@@ -368,9 +299,9 @@ namespace UnrealVS
 			}
 
 			// Call annotate itself
-			bool Result = TryP4Command($"annotate -TcIqu \"{DTE.ActiveDocument.FullName}", PullWorkingDirectoryOn, OutputStdOutOff);
+			bool bResult = TryP4Command($"annotate -TcIqu \"{DTE.ActiveDocument.FullName}", bPullWorkingDirectoryOn, bOutputStdOutOff);
 
-			if (!Result || P4OperationStdErr.Length > 0)
+			if (!bResult || P4OperationStdErr.Length > 0)
 			{
 				P4OutputPane.OutputString($"1>------ P4Annotate call failed");
 				return;
@@ -496,8 +427,8 @@ namespace UnrealVS
 		{
 			if (UnrealVSPackage.Instance.OptionsPage.AllowUnrealVSOverrideDiffSettings)
 			{
-				bool Margin = true;
-				UnrealVSPackage.Instance.EditorOptionsFactory.GlobalOptions.SetOptionValue("Diff/View/ShowDiffOverviewMargin", Margin);
+				bool bMargin = true;
+				UnrealVSPackage.Instance.EditorOptionsFactory.GlobalOptions.SetOptionValue("Diff/View/ShowDiffOverviewMargin", bMargin);
 
 				DifferenceHighlightMode HighlightMode = (DifferenceHighlightMode)3;
 				UnrealVSPackage.Instance.EditorOptionsFactory.GlobalOptions.SetOptionValue("Diff/View/HighlightMode", HighlightMode);
@@ -542,7 +473,7 @@ namespace UnrealVS
 				// p4 print //UE5/Main/Engine/Source/Programs/UnrealVS/UnrealVS.Shared/P4Commands.cs#5 >> file
 
 				string VersionPath = $"{depotPath}#{HaveRev}";
-				if (TryP4Command($"-q print \"{VersionPath}\"",PullWorkingDirectoryOn, OutputStdOutOff))
+				if (TryP4Command($"-q print \"{VersionPath}\"",bPullWorkingDirectoryOn, bOutputStdOutOff))
 				{
 					File.WriteAllText(TempFilePath, P4OperationStdOut);
 
@@ -606,9 +537,9 @@ namespace UnrealVS
 			//P4USER = andrew.firth(config 'd:\p4\frosty\p4config.txt')
 			//P4_perforce:1666_CHARSET = none(set)
 
-			TryP4Command("set", PullWorkingDirectoryOff);
+			TryP4Command("set", bPullWorkingDirectoryOff);
 
-			bool Success = false;
+			bool bSuccess = false;
 
 			if (P4OperationStdOut.Length > 0)
 			{
@@ -622,12 +553,12 @@ namespace UnrealVS
 						string path = lines[0].Split('\'')[1];
 
 						WorkingDirectory = Path.GetDirectoryName(path);
-						Success = true;
+						bSuccess = true;
 					}
 				}
 			}
 
-			if (!Success)
+			if (!bSuccess)
 			{
 				P4OutputPane.OutputString($"attempt to pull p4config.txt info failed{Environment.NewLine}");
 			}
@@ -637,7 +568,7 @@ namespace UnrealVS
 
 		void SetUserInfoStrings()
 		{
-			TryP4Command($"-s -L \"{P4WorkingDirectory}\" info", PullWorkingDirectoryOff);
+			TryP4Command($"-s -L \"{P4WorkingDirectory}\" info", bPullWorkingDirectoryOff);
 
 			Regex UserPattern = new Regex(@"User name: (?<user>.*)$", RegexOptions.Compiled | RegexOptions.Multiline);
 			Regex PortPattern = new Regex(@"Server address: (?<port>.*)$", RegexOptions.Compiled | RegexOptions.Multiline);
@@ -656,7 +587,7 @@ namespace UnrealVS
 			P4OutputPane.OutputString("GetUserInfoStringFull : " + UserInfoComplete);
 		}
 
-		private void PullWorkingDirectory(bool pullfromP4Settings)
+		private void PullWorkingDirectory(bool bPullfromP4Settings)
 		{
 			if (IsSolutionLoaded() && (P4WorkingDirectory == null || P4WorkingDirectory.Length < 2))
 			{
@@ -671,14 +602,14 @@ namespace UnrealVS
 
 			// if the callee wants us to pull a CWD and it wasn't already done
 			// pull it from the p4 config now
-			if (pullfromP4Settings && PullWorkingDirectorFromP4)
+			if (bPullfromP4Settings && bPullWorkingDirectorFromP4)
 			{
 				string NewWorkingDirectory = ReadWorkingDirectorFromP4Config();
 
 				if (NewWorkingDirectory.Length > 1)
 				{
 					P4WorkingDirectory = NewWorkingDirectory;
-					PullWorkingDirectorFromP4 = false;
+					bPullWorkingDirectorFromP4 = false;
 
 					P4OutputPane.OutputString($"P4WorkingDirectory set to '{P4WorkingDirectory}'{Environment.NewLine}");
 				}
@@ -691,29 +622,29 @@ namespace UnrealVS
 			}
 		}
 
-		private bool TryP4Command(string CommandLine, bool PullWorkingDirectoryNow = PullWorkingDirectoryOn, bool OutputStdOut = OutputStdOutOn)
+		private bool TryP4Command(string CommandLine, bool bPullWorkingDirectoryNow = bPullWorkingDirectoryOn, bool bOutputStdOut = bOutputStdOutOn)
 		{
-			return TryP4CommandEx(P4Exe, CommandLine, PullWorkingDirectoryNow, OutputStdOut);
+			return TryP4CommandEx(P4Exe, CommandLine, bPullWorkingDirectoryNow, bOutputStdOut);
 		}
 
-		private bool TryP4VCCommand(string CommandLine, bool PullWorkingDirectoryNow = PullWorkingDirectoryOn, bool OutputStdOut = OutputStdOutOn)
+		private bool TryP4VCCommand(string CommandLine, bool bPullWorkingDirectoryNow = bPullWorkingDirectoryOn, bool bOutputStdOut = bOutputStdOutOn)
 		{
 			if (P4VCCmd.Length > 1)
 			{
-				return TryP4CommandEx(P4VCCmd, CommandLine, PullWorkingDirectoryNow, OutputStdOut);
+				return TryP4CommandEx(P4VCCmd, CommandLine, bPullWorkingDirectoryNow, bOutputStdOut);
 			}
 
 			return false;
 		}
 
-		private bool TryP4VCommand(string CommandLine, bool PullWorkingDirectoryNow = PullWorkingDirectoryOn, bool OutputStdOut = OutputStdOutOn)
+		private bool TryP4VCommand(string CommandLine, bool bPullWorkingDirectoryNow = bPullWorkingDirectoryOn, bool bOutputStdOut = bOutputStdOutOn)
 		{
-			return TryP4CommandEx(P4VExe, CommandLine, PullWorkingDirectoryNow, OutputStdOut);
+			return TryP4CommandEx(P4VExe, CommandLine, bPullWorkingDirectoryNow, bOutputStdOut);
 		}
 
-		private bool TryP4CommandEx(string CmdPath, string CommandLine, bool PullWorkingDirectoryNow = PullWorkingDirectoryOn, bool OutputStdOut = OutputStdOutOn)
+		private bool TryP4CommandEx(string CmdPath, string CommandLine, bool bPullWorkingDirectoryNow = bPullWorkingDirectoryOn, bool bOutputStdOut = bOutputStdOutOn)
 		{
-			PullWorkingDirectory(PullWorkingDirectoryNow);
+			PullWorkingDirectory(bPullWorkingDirectoryNow);
 
 			DTE DTE = UnrealVSPackage.Instance.DTE;
 
@@ -764,7 +695,7 @@ namespace UnrealVS
 			P4OperationStdOut = StdOutSB.ToString();
 			P4OperationStdErr = StdErrSB.ToString();
 
-			if (P4OperationStdOut.Length > 0 && OutputStdOut)
+			if (P4OperationStdOut.Length > 0 && bOutputStdOut)
 			{
 				P4OutputPane.OutputString(P4OperationStdOut);
 			}
