[UnrealVS] - add UBT menu add CompileSingleFile (already existed) and PreprocessSingleFile (slight variance on CompileSingleFile.

diff --git a/Engine/Source/Programs/UnrealVS/UnrealVS.2019/Properties/AssemblyInfo.cs b/Engine/Source/Programs/UnrealVS/UnrealVS.2019/Properties/AssemblyInfo.cs
--- a/Engine/Source/Programs/UnrealVS/UnrealVS.2019/Properties/AssemblyInfo.cs
+++ b/Engine/Source/Programs/UnrealVS/UnrealVS.2019/Properties/AssemblyInfo.cs
@@ -30,8 +30,8 @@ using System.Runtime.InteropServices;
 // You can specify all the values or you can default the Revision and Build Numbers 
 // by using the '*' as shown below:
 
-[assembly: AssemblyVersion("1.66.0.0")]
-[assembly: AssemblyFileVersion("1.66.0.0")]
+[assembly: AssemblyVersion("1.67.0.0")]
+[assembly: AssemblyFileVersion("1.67.0.0")]
 
 
 
diff --git a/Engine/Source/Programs/UnrealVS/UnrealVS.2022/Properties/AssemblyInfo.cs b/Engine/Source/Programs/UnrealVS/UnrealVS.2022/Properties/AssemblyInfo.cs
--- a/Engine/Source/Programs/UnrealVS/UnrealVS.2022/Properties/AssemblyInfo.cs
+++ b/Engine/Source/Programs/UnrealVS/UnrealVS.2022/Properties/AssemblyInfo.cs
@@ -30,8 +30,8 @@ using System.Runtime.InteropServices;
 // You can specify all the values or you can default the Revision and Build Numbers 
 // by using the '*' as shown below:
 
-[assembly: AssemblyVersion("1.66.0.0")]
-[assembly: AssemblyFileVersion("1.66.0.0")]
+[assembly: AssemblyVersion("1.67.0.0")]
+[assembly: AssemblyFileVersion("1.67.0.0")]
 
 
 
diff --git a/Engine/Source/Programs/UnrealVS/UnrealVS.Shared/CompileSingleFile.cs b/Engine/Source/Programs/UnrealVS/UnrealVS.Shared/CompileSingleFile.cs
--- a/Engine/Source/Programs/UnrealVS/UnrealVS.Shared/CompileSingleFile.cs
+++ b/Engine/Source/Programs/UnrealVS/UnrealVS.Shared/CompileSingleFile.cs
@@ -1,6 +1,7 @@
 // Copyright Epic Games, Inc. All Rights Reserved.
 
 using EnvDTE;
+using Microsoft.VisualStudio.Shell;
 using Microsoft.VisualStudio.Shell.Interop;
 using System;
 using System.Collections.Generic;
@@ -8,21 +9,36 @@ using System.ComponentModel.Design;
 using System.Diagnostics;
 using System.IO;
 using System.Windows.Forms;
+using System.Text;
 
 namespace UnrealVS
 {
 	class CompileSingleFile : IDisposable
 	{
-		const int CompileSingleFileButtonID = 0x1075;
+		private const int CompileSingleFileButtonID = 0x1075;
+		private const int PreprocessSingleFileButtonID = 0x1076;
+		private const int UBTSubMenuID = 0x3103;
+		private string	  PPLine = "";
+		private string	  FileToCompileOriginalExt = "";
+
 		static readonly List<string> ValidExtensions = new List<string> { ".c", ".cc", ".cpp", ".cxx" };
 
 		System.Diagnostics.Process ChildProcess;
+		private OleMenuCommand SubMenuCommand;
 
 		public CompileSingleFile()
 		{
 			CommandID CommandID = new CommandID(GuidList.UnrealVSCmdSet, CompileSingleFileButtonID);
 			MenuCommand CompileSingleFileButtonCommand = new MenuCommand(new EventHandler(CompileSingleFileButtonHandler), CommandID);
 			UnrealVSPackage.Instance.MenuCommandService.AddCommand(CompileSingleFileButtonCommand);
+
+			CommandID CommandID2 = new CommandID(GuidList.UnrealVSCmdSet, PreprocessSingleFileButtonID);
+			MenuCommand PreprocessSingleFileButtonCommand = new MenuCommand(new EventHandler(CompileSingleFileButtonHandler), CommandID2);
+			UnrealVSPackage.Instance.MenuCommandService.AddCommand(PreprocessSingleFileButtonCommand);
+
+			// add sub menu for UBT commands
+			SubMenuCommand = new OleMenuCommand(null, new CommandID(GuidList.UnrealVSCmdSet, UBTSubMenuID));
+			UnrealVSPackage.Instance.MenuCommandService.AddCommand(SubMenuCommand);
 		}
 
 		public void Dispose()
@@ -32,7 +48,11 @@ namespace UnrealVS
 
 		void CompileSingleFileButtonHandler(object Sender, EventArgs Args)
 		{
-			if (!TryCompileSingleFile())
+			MenuCommand SenderSubMenuCommand = (MenuCommand)Sender;
+
+			bool PreprocessOnly = SenderSubMenuCommand.CommandID.ID == PreprocessSingleFileButtonID;
+
+			if (!TryCompileSingleFile(PreprocessOnly))
 			{
 				DTE DTE = UnrealVSPackage.Instance.DTE;
 				DTE.ExecuteCommand("Build.Compile");
@@ -53,7 +73,7 @@ namespace UnrealVS
 			}
 		}
 
-		bool TryCompileSingleFile()
+		bool TryCompileSingleFile(bool bPreProcessOnly)
 		{
 			DTE DTE = UnrealVSPackage.Instance.DTE;
 
@@ -159,17 +179,33 @@ namespace UnrealVS
 			DTE.Events.BuildEvents.OnBuildBegin += BuildEvents_OnBuildBegin;
 
 			// Create a delegate for handling output messages
-			void OutputHandler(object Sender, DataReceivedEventArgs Args) { if (Args.Data != null) BuildOutputPane.OutputString($"1>  {Args.Data}{Environment.NewLine}"); }
+			void OutputHandler(object Sender, DataReceivedEventArgs Args) 
+			{ 
+				if (Args.Data != null) 
+				{ 
+					BuildOutputPane.OutputString($"1>  {Args.Data}{Environment.NewLine}"); 
+					if (Args.Data.Contains("PreProcessPath:"))
+					{
+						PPLine = Args.Data;
+					}
+				} 
+			}
 
+			string SolutionDir = Path.GetDirectoryName(UnrealVSPackage.Instance.SolutionFilepath).TrimEnd(Path.DirectorySeparatorChar) + Path.DirectorySeparatorChar;
 			// Get the build command line and escape any environment variables that we use
 			string BuildCommandLine = ActiveNMakeTool.BuildCommandLine;
-			BuildCommandLine = BuildCommandLine.Replace("$(SolutionDir)", Path.GetDirectoryName(UnrealVSPackage.Instance.SolutionFilepath).TrimEnd(Path.DirectorySeparatorChar) + Path.DirectorySeparatorChar);
+			BuildCommandLine = BuildCommandLine.Replace("$(SolutionDir)", SolutionDir);
 			BuildCommandLine = BuildCommandLine.Replace("$(ProjectName)", VCStartupProject.Name);
 
+			PPLine = "";
+			FileToCompileOriginalExt = FileToCompileExt;
+
+			string PreProcess = bPreProcessOnly ? " -NoXGE -Preprocess " : "";
+
 			// Spawn the new process
 			ChildProcess = new System.Diagnostics.Process();
 			ChildProcess.StartInfo.FileName = Path.Combine(Environment.SystemDirectory, "cmd.exe");
-			ChildProcess.StartInfo.Arguments = $"/C \"{BuildCommandLine} -singlefile=\"{FileToCompile}\"\"";
+			ChildProcess.StartInfo.Arguments = $"/C \"{BuildCommandLine} {PreProcess} -singlefile=\"{FileToCompile}\"\"";
 			ChildProcess.StartInfo.WorkingDirectory = Path.GetDirectoryName(StartupProject.FullName);
 			ChildProcess.StartInfo.UseShellExecute = false;
 			ChildProcess.StartInfo.RedirectStandardOutput = true;
@@ -177,12 +213,48 @@ namespace UnrealVS
 			ChildProcess.StartInfo.CreateNoWindow = true;
 			ChildProcess.OutputDataReceived += OutputHandler;
 			ChildProcess.ErrorDataReceived += OutputHandler;
+			if (bPreProcessOnly)
+			{ 
+				// add an event handler to respond to the exit of the preprocess request
+				// and open the generated file if it exists.
+				ChildProcess.EnableRaisingEvents = true;
+				ChildProcess.Exited += new EventHandler(PreprocessExitHandler);
+			}
+			
 			ChildProcess.Start();
 			ChildProcess.BeginOutputReadLine();
 			ChildProcess.BeginErrorReadLine();
+
 			return true;
 		}
 
+		private void PreprocessExitHandler(object sender, System.EventArgs e)
+		{
+			// not all compile actions support pre-process - check it exists
+			if (PPLine.Contains("PreProcessPath:"))
+			{
+				string PPFullPath = PPLine.Replace("PreProcessPath:", "").Trim();
+
+				IVsOutputWindowPane BuildOutputPane = UnrealVSPackage.Instance.GetOutputPane();
+				DTE DTE = UnrealVSPackage.Instance.DTE;
+
+				BuildOutputPane.OutputString($"1>  PPFullPath: {PPFullPath}{Environment.NewLine}");
+
+				if (File.Exists(PPFullPath))
+				{
+					// if the file exists, rename it to isolate the file and have its extension be the original to maintain syntax highlighting
+					string Dir = Path.GetDirectoryName(PPFullPath);
+					string FileName = Path.GetFileNameWithoutExtension(Path.GetFileNameWithoutExtension(PPFullPath)) + "_preprocessed";
+
+					string RenamedFile = Path.Combine(Dir, FileName) + FileToCompileOriginalExt;
+
+					File.Copy(PPFullPath, RenamedFile, true /*overwrite*/);
+
+					DTE.ExecuteCommand("File.OpenFile", $"\"{RenamedFile}\"");
+				}
+			}	
+		}
+
 		private void BuildEvents_OnBuildBegin(vsBuildScope Scope, vsBuildAction Action)
 		{
 			KillChildProcess();
