UnrealVS: Force nuget references in the .vsix

[REDACTED_PROFILE_URL]/Home/issues/5899

diff --git a/Engine/Source/Programs/UnrealVS/UnrealVS.2019/UnrealVS.2019.csproj b/Engine/Source/Programs/UnrealVS/UnrealVS.2019/UnrealVS.2019.csproj
--- a/Engine/Source/Programs/UnrealVS/UnrealVS.2019/UnrealVS.2019.csproj
+++ b/Engine/Source/Programs/UnrealVS/UnrealVS.2019/UnrealVS.2019.csproj
@@ -188,6 +188,11 @@
     <CopyVsixExtensionFiles>False</CopyVsixExtensionFiles>
     <NoWarn>MSB3277</NoWarn>
   </PropertyGroup>
+  <Target Name="IncludePackageReferenceDependencies" AfterTargets="GetVsixSourceItems">
+    <ItemGroup>
+      <VSIXSourceItem Include="@(ReferencePath)" />
+    </ItemGroup>
+  </Target>
   <Import Project="$(MSBuildBinPath)\Microsoft.CSharp.targets" />
   <Import Project="$(VSToolsPath)\VSSDK\Microsoft.VsSDK.targets" Condition="'$(VSToolsPath)' != ''" />
   <Import Project="$(MSBuildExtensionsPath32)\Microsoft\VisualStudio\v11.0\VSSDK\Microsoft.VsSDK.targets" Condition="false" />
diff --git a/Engine/Source/Programs/UnrealVS/UnrealVS.2022/UnrealVS.2022.csproj b/Engine/Source/Programs/UnrealVS/UnrealVS.2022/UnrealVS.2022.csproj
--- a/Engine/Source/Programs/UnrealVS/UnrealVS.2022/UnrealVS.2022.csproj
+++ b/Engine/Source/Programs/UnrealVS/UnrealVS.2022/UnrealVS.2022.csproj
@@ -188,6 +188,11 @@
     <CopyVsixExtensionFiles>False</CopyVsixExtensionFiles>
     <NoWarn>MSB3277</NoWarn>
   </PropertyGroup>
+  <Target Name="IncludePackageReferenceDependencies" AfterTargets="GetVsixSourceItems">
+    <ItemGroup>
+      <VSIXSourceItem Include="@(ReferencePath)" />
+    </ItemGroup>
+  </Target>
   <Import Project="$(MSBuildBinPath)\Microsoft.CSharp.targets" />
   <Import Project="$(VSToolsPath)\VSSDK\Microsoft.VsSDK.targets" Condition="'$(VSToolsPath)' != ''" />
   <Import Project="$(MSBuildExtensionsPath32)\Microsoft\VisualStudio\v11.0\VSSDK\Microsoft.VsSDK.targets" Condition="false" />
