UnrealVS: Set minimum 2019 version to 16.9 and downgrade referenced nuget packages to match.

Updating nuget packages past the VS version prevents the extension from loading correctly.

diff --git a/Engine/Source/Programs/UnrealVS/README b/Engine/Source/Programs/UnrealVS/README
new file mode 100644
--- /dev/null
+++ b/Engine/Source/Programs/UnrealVS/README
@@ -0,0 +1,3 @@
+If nuget packages are updated, the minimum Visual Studio version needs to be updated to match in the .vsixmanifest.
+
+Ex. Microsoft.VisualStudio.SDK 16.9.31025.194 = Minimum Visual Studio version needs to be 16.9
\ No newline at end of file
diff --git a/Engine/Source/Programs/UnrealVS/UnrealVS.2019/Properties/AssemblyInfo.cs b/Engine/Source/Programs/UnrealVS/UnrealVS.2019/Properties/AssemblyInfo.cs
--- a/Engine/Source/Programs/UnrealVS/UnrealVS.2019/Properties/AssemblyInfo.cs
+++ b/Engine/Source/Programs/UnrealVS/UnrealVS.2019/Properties/AssemblyInfo.cs
@@ -30,8 +30,8 @@ using System.Runtime.InteropServices;
 // You can specify all the values or you can default the Revision and Build Numbers 
 // by using the '*' as shown below:
 
-[assembly: AssemblyVersion("1.67.0.0")]
-[assembly: AssemblyFileVersion("1.67.0.0")]
+[assembly: AssemblyVersion("1.68.0.0")]
+[assembly: AssemblyFileVersion("1.68.0.0")]
 
 
 
diff --git a/Engine/Source/Programs/UnrealVS/UnrealVS.2019/UnrealVS.2019.csproj b/Engine/Source/Programs/UnrealVS/UnrealVS.2019/UnrealVS.2019.csproj
--- a/Engine/Source/Programs/UnrealVS/UnrealVS.2019/UnrealVS.2019.csproj
+++ b/Engine/Source/Programs/UnrealVS/UnrealVS.2019/UnrealVS.2019.csproj
@@ -136,18 +136,18 @@
       <Version>4.7.0</Version>
     </PackageReference>
     <PackageReference Include="Microsoft.VisualStudio.Sdk">
-      <Version>16.10.31321.278</Version>
+      <Version>16.9.31025.194</Version>
     </PackageReference>
     <PackageReference Include="Microsoft.VisualStudio.Threading">
-      <Version>16.10.56</Version>
+      <Version>16.9.54</Version>
     </PackageReference>
     <PackageReference Include="Microsoft.VisualStudio.Threading.Analyzers">
-      <Version>16.10.56</Version>
+      <Version>16.9.54</Version>
       <IncludeAssets>runtime; build; native; contentfiles; analyzers; buildtransitive</IncludeAssets>
       <PrivateAssets>all</PrivateAssets>
     </PackageReference>
     <PackageReference Include="Microsoft.VisualStudio.VCProjectEngine">
-      <Version>16.10.31320.204</Version>
+      <Version>16.9.31025.104</Version>
     </PackageReference>
     <PackageReference Include="Microsoft.VSSDK.BuildTools">
       <Version>17.0.1600</Version>
@@ -188,6 +188,11 @@
     <CopyVsixExtensionFiles>False</CopyVsixExtensionFiles>
     <NoWarn>MSB3277</NoWarn>
   </PropertyGroup>
+  <Target Name="IncludeNuGetPackageReferences" AfterTargets="GetVsixSourceItems">
+    <ItemGroup>
+      <VSIXSourceItem Include="@(ReferenceCopyLocalPaths)" Condition="'%(ReferenceCopyLocalPaths.NuGetPackageId)' != ''" />
+    </ItemGroup>
+  </Target>
   <Import Project="$(MSBuildBinPath)\Microsoft.CSharp.targets" />
   <Import Project="$(VSToolsPath)\VSSDK\Microsoft.VsSDK.targets" Condition="'$(VSToolsPath)' != ''" />
   <Import Project="$(MSBuildExtensionsPath32)\Microsoft\VisualStudio\v11.0\VSSDK\Microsoft.VsSDK.targets" Condition="false" />
diff --git a/Engine/Source/Programs/UnrealVS/UnrealVS.2022/Properties/AssemblyInfo.cs b/Engine/Source/Programs/UnrealVS/UnrealVS.2022/Properties/AssemblyInfo.cs
--- a/Engine/Source/Programs/UnrealVS/UnrealVS.2022/Properties/AssemblyInfo.cs
+++ b/Engine/Source/Programs/UnrealVS/UnrealVS.2022/Properties/AssemblyInfo.cs
@@ -30,8 +30,8 @@ using System.Runtime.InteropServices;
 // You can specify all the values or you can default the Revision and Build Numbers 
 // by using the '*' as shown below:
 
-[assembly: AssemblyVersion("1.67.0.0")]
-[assembly: AssemblyFileVersion("1.67.0.0")]
+[assembly: AssemblyVersion("1.68.0.0")]
+[assembly: AssemblyFileVersion("1.68.0.0")]
 
 
 
diff --git a/Engine/Source/Programs/UnrealVS/UnrealVS.2022/UnrealVS.2022.csproj b/Engine/Source/Programs/UnrealVS/UnrealVS.2022/UnrealVS.2022.csproj
--- a/Engine/Source/Programs/UnrealVS/UnrealVS.2022/UnrealVS.2022.csproj
+++ b/Engine/Source/Programs/UnrealVS/UnrealVS.2022/UnrealVS.2022.csproj
@@ -125,7 +125,7 @@
   </ItemGroup>
   <ItemGroup>
     <PackageReference Include="Microsoft.CodeAnalysis.CSharp">
-      <Version>3.11.0-1.final</Version>
+      <Version>4.0.0-1.final</Version>
     </PackageReference>
     <PackageReference Include="Microsoft.CodeAnalysis.NetAnalyzers">
       <Version>5.0.3</Version>
@@ -135,8 +135,7 @@
     <PackageReference Include="Microsoft.CSharp">
       <Version>4.7.0</Version>
     </PackageReference>
-    <PackageReference Include="Microsoft.VisualStudio.Sdk">
-      <Version>17.0.0-previews-1-31410-273</Version>
+    <PackageReference Include="Microsoft.VisualStudio.Sdk" Version="17.0.0-previews-1-31410-273">
     </PackageReference>
     <PackageReference Include="Microsoft.VisualStudio.Threading">
       <Version>17.0.17-alpha</Version>
@@ -149,8 +148,7 @@
     <PackageReference Include="Microsoft.VisualStudio.VCProjectEngine">
       <Version>17.0.0-previews-1-31410-258</Version>
     </PackageReference>
-    <PackageReference Include="Microsoft.VSSDK.BuildTools">
-      <Version>17.0.1619-preview1</Version>
+    <PackageReference Include="Microsoft.VSSDK.BuildTools" Version="17.0.3154-preview3">
       <IncludeAssets>runtime; build; native; contentfiles; analyzers; buildtransitive</IncludeAssets>
       <PrivateAssets>all</PrivateAssets>
     </PackageReference>
@@ -188,6 +186,11 @@
     <CopyVsixExtensionFiles>False</CopyVsixExtensionFiles>
     <NoWarn>MSB3277</NoWarn>
   </PropertyGroup>
+  <Target Name="IncludeNuGetPackageReferences" AfterTargets="GetVsixSourceItems">
+    <ItemGroup>
+      <VSIXSourceItem Include="@(ReferenceCopyLocalPaths)" Condition="'%(ReferenceCopyLocalPaths.NuGetPackageId)' != ''" />
+    </ItemGroup>
+  </Target>
   <Import Project="$(MSBuildBinPath)\Microsoft.CSharp.targets" />
   <Import Project="$(VSToolsPath)\VSSDK\Microsoft.VsSDK.targets" Condition="'$(VSToolsPath)' != ''" />
   <Import Project="$(MSBuildExtensionsPath32)\Microsoft\VisualStudio\v11.0\VSSDK\Microsoft.VsSDK.targets" Condition="false" />
