[UnrealVS] - fixes and cleanup

- add option to allow async checkout
- resolved problems that caused with many commands running at once
- fix issue with P4V command waiting for exit
- resolve issue where Checkout on Edit was not honoring the user option
- add ReleaseNotes to the plugin
- add UpdateInstructions for future engineers

diff --git a/Engine/Source/Programs/UnrealVS/ReleaseNotes.txt b/Engine/Source/Programs/UnrealVS/ReleaseNotes.txt
new file mode 100644
--- /dev/null
+++ b/Engine/Source/Programs/UnrealVS/ReleaseNotes.txt
@@ -0,0 +1,6 @@
+1.70
+* add option to allow async checkout
+* resolved concurrent command issues
+* fix issue with P4V commands waiting for exit
+* resolve issue where Checkout on Edit was not honoring the user option
+
diff --git a/Engine/Source/Programs/UnrealVS/UnrealVS.2019/Properties/AssemblyInfo.cs b/Engine/Source/Programs/UnrealVS/UnrealVS.2019/Properties/AssemblyInfo.cs
--- a/Engine/Source/Programs/UnrealVS/UnrealVS.2019/Properties/AssemblyInfo.cs
+++ b/Engine/Source/Programs/UnrealVS/UnrealVS.2019/Properties/AssemblyInfo.cs
@@ -11,6 +11,7 @@ using System.Runtime.InteropServices;
 [assembly: AssemblyTitle("UnrealVS")]
 [assembly: AssemblyDescription("")]
 [assembly: AssemblyConfiguration("")]
+
 [assembly: AssemblyCompany("Epic Games, Inc.")]
 [assembly: AssemblyProduct("UnrealVS")]
 [assembly: AssemblyCopyright("Copyright Epic Games, Inc.")]
@@ -30,8 +31,8 @@ using System.Runtime.InteropServices;
 // You can specify all the values or you can default the Revision and Build Numbers 
 // by using the '*' as shown below:
 
-[assembly: AssemblyVersion("1.69.0.0")]
-[assembly: AssemblyFileVersion("1.69.0.0")]
+[assembly: AssemblyVersion("1.70.0.0")]
+[assembly: AssemblyFileVersion("1.70.0.0")]
 
 
 
diff --git a/Engine/Source/Programs/UnrealVS/UnrealVS.2019/UnrealVS.2019.csproj b/Engine/Source/Programs/UnrealVS/UnrealVS.2019/UnrealVS.2019.csproj
--- a/Engine/Source/Programs/UnrealVS/UnrealVS.2019/UnrealVS.2019.csproj
+++ b/Engine/Source/Programs/UnrealVS/UnrealVS.2019/UnrealVS.2019.csproj
@@ -96,6 +96,10 @@
     <None Include="Resources\BatchBuilder.png" />
     <Resource Include="Resources\BatchBuilderBackground.png" />
     <Resource Include="Resources\Icon.ico" />
+    <Content Include="ReleaseNotes.txt">
+      <CopyToOutputDirectory>Always</CopyToOutputDirectory>
+      <IncludeInVSIX>true</IncludeInVSIX>
+    </Content>
     <Content Include="Resources\Icon.png">
       <CopyToOutputDirectory>Always</CopyToOutputDirectory>
       <IncludeInVSIX>true</IncludeInVSIX>
diff --git a/Engine/Source/Programs/UnrealVS/UnrealVS.2022/Properties/AssemblyInfo.cs b/Engine/Source/Programs/UnrealVS/UnrealVS.2022/Properties/AssemblyInfo.cs
--- a/Engine/Source/Programs/UnrealVS/UnrealVS.2022/Properties/AssemblyInfo.cs
+++ b/Engine/Source/Programs/UnrealVS/UnrealVS.2022/Properties/AssemblyInfo.cs
@@ -30,8 +30,8 @@ using System.Runtime.InteropServices;
 // You can specify all the values or you can default the Revision and Build Numbers 
 // by using the '*' as shown below:
 
-[assembly: AssemblyVersion("1.69.0.0")]
-[assembly: AssemblyFileVersion("1.69.0.0")]
+[assembly: AssemblyVersion("1.70.0.0")]
+[assembly: AssemblyFileVersion("1.70.0.0")]
 
 
 
diff --git a/Engine/Source/Programs/UnrealVS/UnrealVS.2022/UnrealVS.2022.csproj b/Engine/Source/Programs/UnrealVS/UnrealVS.2022/UnrealVS.2022.csproj
--- a/Engine/Source/Programs/UnrealVS/UnrealVS.2022/UnrealVS.2022.csproj
+++ b/Engine/Source/Programs/UnrealVS/UnrealVS.2022/UnrealVS.2022.csproj
@@ -96,6 +96,10 @@
     <None Include="Resources\BatchBuilder.png" />
     <Resource Include="Resources\BatchBuilderBackground.png" />
     <Resource Include="Resources\Icon.ico" />
+    <Content Include="ReleaseNotes.txt">
+      <CopyToOutputDirectory>Always</CopyToOutputDirectory>
+      <IncludeInVSIX>true</IncludeInVSIX>
+    </Content>
     <Content Include="Resources\Icon.png">
       <CopyToOutputDirectory>Always</CopyToOutputDirectory>
       <IncludeInVSIX>true</IncludeInVSIX>
diff --git a/Engine/Source/Programs/UnrealVS/UnrealVS.Shared/P4Commands.cs b/Engine/Source/Programs/UnrealVS/UnrealVS.Shared/P4Commands.cs
--- a/Engine/Source/Programs/UnrealVS/UnrealVS.Shared/P4Commands.cs
+++ b/Engine/Source/Programs/UnrealVS/UnrealVS.Shared/P4Commands.cs
@@ -12,6 +12,8 @@ using System.Diagnostics;
 using System.IO;
 using System.Text;
 using System.Text.RegularExpressions;
+using System.Threading;
+using System.Threading.Tasks;
 
 namespace UnrealVS
 {
@@ -23,6 +25,9 @@ namespace UnrealVS
 		private const bool bOutputStdOutOn = true;
 		private const bool bOutputStdOutOff = false;
 
+		private const bool bWaitForCompletionOn = true;
+		private const bool bWaitForCompletionOff = false;
+
 		private const int P4SubMenuID = 0x3100;
 		private const int P4CheckoutButtonID = 0x1450;
 		private const int P4AnnotateButtonID = 0x1451;
@@ -34,13 +39,17 @@ namespace UnrealVS
 
 		private OleMenuCommand SubMenuCommand;
 
-		private System.Diagnostics.Process ChildProcess;
+		//private System.Diagnostics.Process ChildProcess;
+
+		private List<System.Diagnostics.Process> ChildProcessList = new List<System.Diagnostics.Process>();
+
 		private IVsOutputWindowPane P4OutputPane;
 		private string P4WorkingDirectory;
 		private bool bPullWorkingDirectorFromP4 = true;
+		private static Mutex bPullWorkingDirectorFromP4Mutex = new Mutex();
 
 		// stdXX from last operation
-		private string P4OperationStdOut;
+		//private string P4OperationStdOut;
 		private string P4OperationStdErr;
 
 		// Exe paths
@@ -197,34 +206,6 @@ namespace UnrealVS
 			}
 		}
 
-		private void EditProjectRecursive(Project Proj)
-		{
-			if (!Proj.Saved)
-			{
-				OpenForEdit(Proj.FullName);
-			}
-
-			if (Proj.ProjectItems == null)
-			{
-				return;
-			}
-
-			foreach (ProjectItem ProjItem in Proj.ProjectItems)
-			{
-				if (ProjItem.SubProject != null)
-				{
-					EditProjectRecursive(ProjItem.SubProject);
-				}
-				else if (!ProjItem.Saved)
-				{
-					for (short Index = 1; Index <= ProjItem.FileCount; Index++)
-					{
-						OpenForEdit(ProjItem.get_FileNames(Index));
-					}
-				}
-			}
-		}
-
 		private void UpdateMenuOptions()
 		{
 			DTE DTE = UnrealVSPackage.Instance.DTE;
@@ -251,15 +232,14 @@ namespace UnrealVS
 
 		private void KillChildProcess()
 		{
-			if (ChildProcess != null)
+			foreach(var OldProcess in ChildProcessList)
 			{
-				if (!ChildProcess.HasExited)
+				if (!OldProcess.HasExited)
 				{
-					ChildProcess.Kill();
-					ChildProcess.WaitForExit();
+					OldProcess.Kill();
+					OldProcess.WaitForExit();
 				}
-				ChildProcess.Dispose();
-				ChildProcess = null;
+				OldProcess.Dispose();
 			}
 		}
 
@@ -298,8 +278,10 @@ namespace UnrealVS
 				return;
 			}
 
+			string CaptureP4Output = "";
+
 			// Call annotate itself
-			bool bResult = TryP4Command($"annotate -TcIqu \"{DTE.ActiveDocument.FullName}", bPullWorkingDirectoryOn, bOutputStdOutOff);
+			bool bResult = TryP4Command($"annotate -TcIqu \"{DTE.ActiveDocument.FullName}", out CaptureP4Output, bPullWorkingDirectoryOn, bOutputStdOutOff);
 
 			if (!bResult || P4OperationStdErr.Length > 0)
 			{
@@ -336,9 +318,9 @@ namespace UnrealVS
 				// /*	13149436:            First.Last 2020/05/04*/
 
 				// This is the per line offset of the annotation added to the document
-				int AnnotateOffset = P4OperationStdOut.IndexOf(FirstLine);
+				int AnnotateOffset = CaptureP4Output.IndexOf(FirstLine);
 
-				string[] AnnotateLines = P4OperationStdOut.Split('\n');
+				string[] AnnotateLines = CaptureP4Output.Split('\n');
 
 				foreach (string Line in AnnotateLines)
 				{
@@ -450,17 +432,19 @@ namespace UnrealVS
 
 			ChangeDiffSetting();
 
+			string CaptureP4Output = "";
+
 			// get the HAVE revision
 			// p4 fstat -T "haveRev" -Olp //UE5/Main/Engine/Source/Programs/UnrealVS/UnrealVS.Shared/P4Commands.cs
-			if (TryP4Command($"fstat -T \"haveRev,depotFile\" -Olp \"{DTE.ActiveDocument.FullName}\""))
+			if (TryP4Command($"fstat -T \"haveRev,depotFile\" -Olp \"{DTE.ActiveDocument.FullName}\"", out CaptureP4Output))
 			{
 				// expect output of the form
 				//		"... haveRev 5"
 				//		"... depotFile //UE5/Main/Engine/Source/Programs/UnrealVS/UnrealVS.Shared/P4Commands.cs
 				Regex HavePattern = new Regex(@"... haveRev (?<Have>.*)$", RegexOptions.Compiled | RegexOptions.Multiline);
 				Regex DepotPathPattern = new Regex(@"... depotFile (?<depotFile>.*)$", RegexOptions.Compiled | RegexOptions.Multiline);
-				System.Text.RegularExpressions.Match HaveMatch = HavePattern.Match(P4OperationStdOut);
-				System.Text.RegularExpressions.Match PathMatch = DepotPathPattern.Match(P4OperationStdOut);
+				System.Text.RegularExpressions.Match HaveMatch = HavePattern.Match(CaptureP4Output);
+				System.Text.RegularExpressions.Match PathMatch = DepotPathPattern.Match(CaptureP4Output);
 				int HaveRev = Int32.Parse(HaveMatch.Groups["Have"].Value.Trim());
 				string depotPath = PathMatch.Groups["depotFile"].Value.Trim();
 
@@ -472,10 +456,12 @@ namespace UnrealVS
 				// sync the HAVE revision to a file
 				// p4 print //UE5/Main/Engine/Source/Programs/UnrealVS/UnrealVS.Shared/P4Commands.cs#5 >> file
 
+				string CaptureP4Output2 = "";
+
 				string VersionPath = $"{depotPath}#{HaveRev}";
-				if (TryP4Command($"-q print \"{VersionPath}\"",bPullWorkingDirectoryOn, bOutputStdOutOff))
+				if (TryP4Command($"-q print \"{VersionPath}\"",out CaptureP4Output2, bPullWorkingDirectoryOn, bOutputStdOutOff))
 				{
-					File.WriteAllText(TempFilePath, P4OperationStdOut);
+					File.WriteAllText(TempFilePath, CaptureP4Output2);
 
 					// Tools.DiffFiles SourceFile, TargetFile, [SourceDisplayName],[TargetDisplayName]
 					DTE.ExecuteCommand("Tools.DiffFiles", $"\"{TempFilePath}\" \"{DTE.ActiveDocument.FullName}\" \"{VersionPath}\" \"{DTE.ActiveDocument.FullName}\"");
@@ -493,10 +479,10 @@ namespace UnrealVS
 			string ChangesTempFileName = Path.GetFileNameWithoutExtension(DTE.ActiveDocument.FullName) + "_last_10_changelists" + Path.GetExtension(DTE.ActiveDocument.FullName);
 			string ChangesTempFilePath = Path.Combine(TempPath, ChangesTempFileName);
 
-
-			if (TryP4Command($"changes -itls submitted -m 10 {DTE.ActiveDocument.FullName}"))
+			string CaptureP4Output = "";
+			if (TryP4Command($"changes -itls submitted -m 10 {DTE.ActiveDocument.FullName}", out CaptureP4Output, bOutputStdOutOff))
 			{
-				File.WriteAllText(ChangesTempFilePath, P4OperationStdOut);
+				File.WriteAllText(ChangesTempFilePath, CaptureP4Output);
 
 				DTE.ExecuteCommand("File.OpenFile", $"\"{ChangesTempFilePath}\"");
 				DTE.ActiveDocument.Activate();
@@ -515,15 +501,38 @@ namespace UnrealVS
 
 			TryP4VCommand($"-s {DTE.ActiveDocument.FullName}");
 		}
-		public void OpenForEdit(string FileName)
+		public void OpenForEdit(string FileName, bool CheckOptionsFlag = false)
 		{
+			// if the callee requested it, honor the enabling of autocheckout option - return if its off
+			if (CheckOptionsFlag && !UnrealVSPackage.Instance.OptionsPage.AllowUnrealVSCheckoutOnEdit)
+			{
+				//P4OutputPane.OutputString($"autocheckout disabled: {FileName}");
+				return;
+			}	
+
 			// Don't open for edit if the file is already writable
 			if (!File.Exists(FileName) || !File.GetAttributes(FileName).HasFlag(FileAttributes.ReadOnly))
 			{
+				//P4OutputPane.OutputString($"already writeable: {FileName}");
 				return;
 			}
 
-			TryP4Command($"edit \"{FileName}");
+			if (UnrealVSPackage.Instance.OptionsPage.AllowAsyncP4Checkout)
+			{
+				//P4OutputPane.OutputString($"Marking file as writeable: {FileName}");
+				// Mark the file as writeable
+				FileAttributes NewAttributes = File.GetAttributes(FileName) & ~FileAttributes.ReadOnly;
+				File.SetAttributes(FileName, NewAttributes);
+
+				// then send a fire and forget p4 edit command - does not lock the UI
+				TryP4CommandAsync($"edit \"{FileName}");
+			}
+			else
+			{
+				string CaptureP4Output = "";
+				// sync edit command - will lock the UI but be safer (no risk of a locally writeable file that is not checked out)
+				TryP4Command($"edit \"{FileName}", out CaptureP4Output);
+			}
 		}
 
 		private string ReadWorkingDirectorFromP4Config()
@@ -537,13 +546,14 @@ namespace UnrealVS
 			//P4USER = andrew.firth(config 'd:\p4\frosty\p4config.txt')
 			//P4_perforce:1666_CHARSET = none(set)
 
-			TryP4Command("set", bPullWorkingDirectoryOff);
+			string CaptureP4Output = "";
+			TryP4Command("set", out CaptureP4Output, bPullWorkingDirectoryOff);
 
 			bool bSuccess = false;
 
-			if (P4OperationStdOut.Length > 0)
+			if (CaptureP4Output.Length > 0)
 			{
-				string[] lines = P4OperationStdOut.Split('\n');
+				string[] lines = CaptureP4Output.Split('\n');
 
 				// very brittle - Index assume there is a better way leveraging C#
 				if (lines.Length > 0)
@@ -568,15 +578,16 @@ namespace UnrealVS
 
 		void SetUserInfoStrings()
 		{
-			TryP4Command($"-s -L \"{P4WorkingDirectory}\" info", bPullWorkingDirectoryOff);
+			string CaptureP4Output = "";
+			TryP4Command($"-s -L \"{P4WorkingDirectory}\" info", out CaptureP4Output, bPullWorkingDirectoryOff);
 
 			Regex UserPattern = new Regex(@"User name: (?<user>.*)$", RegexOptions.Compiled | RegexOptions.Multiline);
 			Regex PortPattern = new Regex(@"Server address: (?<port>.*)$", RegexOptions.Compiled | RegexOptions.Multiline);
 			Regex ClientPattern = new Regex(@"Client name: (?<client>.*)$", RegexOptions.Compiled | RegexOptions.Multiline);
 
-			System.Text.RegularExpressions.Match UserMatch = UserPattern.Match(P4OperationStdOut);
-			System.Text.RegularExpressions.Match PortMatch = PortPattern.Match(P4OperationStdOut);
-			System.Text.RegularExpressions.Match ClientMath = ClientPattern.Match(P4OperationStdOut);
+			System.Text.RegularExpressions.Match UserMatch = UserPattern.Match(CaptureP4Output);
+			System.Text.RegularExpressions.Match PortMatch = PortPattern.Match(CaptureP4Output);
+			System.Text.RegularExpressions.Match ClientMath = ClientPattern.Match(CaptureP4Output);
 
 			Port = PortMatch.Groups["port"].Value.Trim();
 			Username = UserMatch.Groups["user"].Value.Trim();
@@ -584,11 +595,13 @@ namespace UnrealVS
 
 			UserInfoComplete = string.Format(" -p {0} -u {1} -c {2} ", Port, Username, Client);
 
-			P4OutputPane.OutputString("GetUserInfoStringFull : " + UserInfoComplete);
+			P4OutputPane.OutputString("GetUserInfoStringFull : " + UserInfoComplete + Environment.NewLine);
 		}
 
 		private void PullWorkingDirectory(bool bPullfromP4Settings)
 		{
+			bPullWorkingDirectorFromP4Mutex.WaitOne();
+
 			if (IsSolutionLoaded() && (P4WorkingDirectory == null || P4WorkingDirectory.Length < 2))
 			{
 				DTE DTE = UnrealVSPackage.Instance.DTE;
@@ -617,21 +630,28 @@ namespace UnrealVS
 				{
 					P4OutputPane.OutputString($"P4WorkingDirectory set failed {Environment.NewLine}");
 				}
+			}
 
+			bPullWorkingDirectorFromP4Mutex.ReleaseMutex();
+		}
 
-			}
+		public async Task<bool> TryP4CommandAsync(string CommandLine, bool bPullWorkingDirectoryNow = bPullWorkingDirectoryOn, bool bOutputStdOut = bOutputStdOutOn)
+		{
+			string CaptureP4Output;
+			return await System.Threading.Tasks.Task.Run( () => TryP4Command(CommandLine, out CaptureP4Output, bPullWorkingDirectoryNow, bOutputStdOut) );
 		}
 
-		private bool TryP4Command(string CommandLine, bool bPullWorkingDirectoryNow = bPullWorkingDirectoryOn, bool bOutputStdOut = bOutputStdOutOn)
+		private bool TryP4Command(string CommandLine, out string CaptureP4Output, bool bPullWorkingDirectoryNow = bPullWorkingDirectoryOn, bool bOutputStdOut = bOutputStdOutOn)
 		{
-			return TryP4CommandEx(P4Exe, CommandLine, bPullWorkingDirectoryNow, bOutputStdOut);
+			return TryP4CommandEx(P4Exe, CommandLine, out CaptureP4Output, bPullWorkingDirectoryNow, bOutputStdOut);
 		}
 
 		private bool TryP4VCCommand(string CommandLine, bool bPullWorkingDirectoryNow = bPullWorkingDirectoryOn, bool bOutputStdOut = bOutputStdOutOn)
 		{
 			if (P4VCCmd.Length > 1)
 			{
-				return TryP4CommandEx(P4VCCmd, CommandLine, bPullWorkingDirectoryNow, bOutputStdOut);
+				string ignore;
+				return TryP4CommandEx(P4VCCmd, CommandLine, out ignore, bPullWorkingDirectoryNow, bOutputStdOut);
 			}
 
 			return false;
@@ -639,38 +659,39 @@ namespace UnrealVS
 
 		private bool TryP4VCommand(string CommandLine, bool bPullWorkingDirectoryNow = bPullWorkingDirectoryOn, bool bOutputStdOut = bOutputStdOutOn)
 		{
-			return TryP4CommandEx(P4VExe, CommandLine, bPullWorkingDirectoryNow, bOutputStdOut);
+			string ignore;
+			return TryP4CommandEx(P4VExe, CommandLine, out ignore, bPullWorkingDirectoryNow, bOutputStdOut, bWaitForCompletionOff);
 		}
 
-		private bool TryP4CommandEx(string CmdPath, string CommandLine, bool bPullWorkingDirectoryNow = bPullWorkingDirectoryOn, bool bOutputStdOut = bOutputStdOutOn)
+		private bool TryP4CommandEx(string CmdPath, string CommandLine, out string CaptureP4Output, bool bPullWorkingDirectoryNow = bPullWorkingDirectoryOn, bool bOutputStdOut = bOutputStdOutOn, bool bWaitForCompletion = bWaitForCompletionOn)
 		{
 			PullWorkingDirectory(bPullWorkingDirectoryNow);
 
 			DTE DTE = UnrealVSPackage.Instance.DTE;
 
 			// If there's already a operation in progress, abort... potential issues here
-			// in the OnSave/OnSaveAll operations.
-			if (ChildProcess != null && !ChildProcess.HasExited)
+			// in the OnSave/OnSaveAll operations where we will stall.
+/*			if (ChildProcess != null && !ChildProcess.HasExited)
 			{
 				P4OutputPane.OutputString($"P4 operation already in flight...waiting{Environment.NewLine}");
 				P4OutputPane.Activate();
 				ChildProcess.WaitForExit();
-			}
+			}*/
 
 			// Make sure any existing op is stopped
-			KillChildProcess();
+			//KillChildProcess();
 
 			// Set up the output pane
 			P4OutputPane.Activate();
-			P4OutputPane.OutputString($"1>------ P4 Operation started{Environment.NewLine}");
-			P4OutputPane.OutputString($"CWD: {P4WorkingDirectory}{Environment.NewLine}");
-			P4OutputPane.OutputString($"CMD: {P4Exe} {CommandLine}{Environment.NewLine}");
+			//P4OutputPane.OutputString($"1>------ P4 Operation started{Environment.NewLine}");
+			//P4OutputPane.OutputString($"CWD: {P4WorkingDirectory}{Environment.NewLine}");
+			//P4OutputPane.OutputString($"CMD: {P4Exe} {CommandLine}{Environment.NewLine}");
 
 			StringBuilder StdOutSB = new StringBuilder();
 			StringBuilder StdErrSB = new StringBuilder();
 
 			// Spawn the new process
-			ChildProcess = new System.Diagnostics.Process()
+			System.Diagnostics.Process ChildProcess = new System.Diagnostics.Process()
 			{
 				StartInfo = new ProcessStartInfo()
 				{
@@ -686,33 +707,67 @@ namespace UnrealVS
 			// Create a delegate for handling output messages
 			ChildProcess.OutputDataReceived += (s, a) => { if (a.Data != null) StdOutSB.AppendLine(a.Data); };
 			ChildProcess.ErrorDataReceived += (s, a) => { if (a.Data != null) StdErrSB.AppendLine(a.Data); };
+			if (bWaitForCompletion == false)
+			{
+				ChildProcess.EnableRaisingEvents = true;
+				ChildProcess.Exited += (s, a) =>
+				{
+					string P4Output = StdOutSB.ToString();
+					P4OperationStdErr = StdErrSB.ToString();
+
+					if (P4Output.Length > 0 && bOutputStdOut)
+					{
+						P4OutputPane.OutputString(P4Output);
+					}
+
+					if (P4OperationStdErr.Length > 0)
+					{
+						P4OutputPane.OutputString(P4OperationStdErr);
+					}
+
+					//P4OutputPane.OutputString($"1>------ P4 Operation complete{Environment.NewLine}");
+
+					ChildProcessList.Remove(ChildProcess);
+				};
+			}
+
+			ChildProcessList.Add(ChildProcess);
+
 			ChildProcess.Start();
 			ChildProcess.BeginOutputReadLine();
 			ChildProcess.BeginErrorReadLine();
 
-			ChildProcess.WaitForExit();
+			if (bWaitForCompletion)
+			{ 
+				ChildProcess.WaitForExit();
 
-			P4OperationStdOut = StdOutSB.ToString();
-			P4OperationStdErr = StdErrSB.ToString();
+				CaptureP4Output = StdOutSB.ToString();
+				P4OperationStdErr = StdErrSB.ToString();
 
-			if (P4OperationStdOut.Length > 0 && bOutputStdOut)
-			{
-				P4OutputPane.OutputString(P4OperationStdOut);
-			}
+				if (CaptureP4Output.Length > 0 && bOutputStdOut)
+				{
+					P4OutputPane.OutputString(CaptureP4Output);
+				}
 
-			if (P4OperationStdErr.Length > 0)
-			{
-				P4OutputPane.OutputString(P4OperationStdErr);
-			}
+				if (P4OperationStdErr.Length > 0)
+				{
+					P4OutputPane.OutputString(P4OperationStdErr);
+				}
+
+				//P4OutputPane.OutputString($"1>------ P4 Operation complete{Environment.NewLine}");
 
-			P4OutputPane.OutputString($"1>------ P4 Operation complete{Environment.NewLine}");
+				ChildProcessList.Remove(ChildProcess);
 
-			return P4OperationStdErr.Length == 0;
+				return P4OperationStdErr.Length == 0;
+			}
+			else
+			{
+				CaptureP4Output = "async operations do not capture p4 output";
+				return true;
+			}
 		}
 	}
 
-	// 
-
 	internal class IntercepteSave : IVsRunningDocTableEvents3
 	{
 		private readonly DTE DTE;
@@ -732,7 +787,7 @@ namespace UnrealVS
 
 			string AbsoluteFilePath = DocumentInfo.Moniker;
 
-			P4Ops.OpenForEdit(AbsoluteFilePath);
+			P4Ops.OpenForEdit(AbsoluteFilePath, true);
 
 			return VSConstants.S_OK;
 		}
diff --git a/Engine/Source/Programs/UnrealVS/UnrealVS.Shared/UnrealVSOptions.cs b/Engine/Source/Programs/UnrealVS/UnrealVS.Shared/UnrealVSOptions.cs
--- a/Engine/Source/Programs/UnrealVS/UnrealVS.Shared/UnrealVSOptions.cs
+++ b/Engine/Source/Programs/UnrealVS/UnrealVS.Shared/UnrealVSOptions.cs
@@ -32,6 +32,10 @@ namespace UnrealVS
 		[Description("Uses p4 ini / environment settings to call P4 functionlality, use 'RunUAT P4WriteConfig' to initialize those settings")]
 		public bool AllowUnrealVSP4 { get; set; }
 
+		[Category("Unreal.P4")]
+		[DisplayName("Allow Async Checkout")]
+		[Description("Switches checkout (on edit or manual) to an async method, note if the p4 operation fails the local file will still be writeable")]
+		public bool AllowAsyncP4Checkout { get; set; }
 		protected override void OnApply(PageApplyEventArgs e)
 		{
 			base.OnApply(e);
diff --git a/Engine/Source/Programs/UnrealVS/UpdateInstructions.txt b/Engine/Source/Programs/UnrealVS/UpdateInstructions.txt
new file mode 100644
--- /dev/null
+++ b/Engine/Source/Programs/UnrealVS/UpdateInstructions.txt
@@ -0,0 +1,5 @@
+1. Make your changes
+2. Update version number in both VSIX manifest files and both AssemblyInfo files.
+3. update ReleaseNotes.txt with the changes
+4. Checkout the VSIX files and rebuild 2019 IN VS2019, 2022 in VS2022
+5. Checkin all the above in 1 CL
