UnrealVS: Fix compile errors due to ReleaseNotes.txt not existing in separate UnrealVS.2019 and UnrealVS.2022 directories.

diff --git a/Engine/Source/Programs/UnrealVS/UnrealVS.2019/UnrealVS.2019.csproj b/Engine/Source/Programs/UnrealVS/UnrealVS.2019/UnrealVS.2019.csproj
--- a/Engine/Source/Programs/UnrealVS/UnrealVS.2019/UnrealVS.2019.csproj
+++ b/Engine/Source/Programs/UnrealVS/UnrealVS.2019/UnrealVS.2019.csproj
@@ -96,7 +96,8 @@
     <None Include="Resources\BatchBuilder.png" />
     <Resource Include="Resources\BatchBuilderBackground.png" />
     <Resource Include="Resources\Icon.ico" />
-    <Content Include="ReleaseNotes.txt">
+    <Content Include="..\ReleaseNotes.txt">
+	  <Visible>false</Visible>
       <CopyToOutputDirectory>Always</CopyToOutputDirectory>
       <IncludeInVSIX>true</IncludeInVSIX>
     </Content>
diff --git a/Engine/Source/Programs/UnrealVS/UnrealVS.2022/UnrealVS.2022.csproj b/Engine/Source/Programs/UnrealVS/UnrealVS.2022/UnrealVS.2022.csproj
--- a/Engine/Source/Programs/UnrealVS/UnrealVS.2022/UnrealVS.2022.csproj
+++ b/Engine/Source/Programs/UnrealVS/UnrealVS.2022/UnrealVS.2022.csproj
@@ -96,7 +96,8 @@
     <None Include="Resources\BatchBuilder.png" />
     <Resource Include="Resources\BatchBuilderBackground.png" />
     <Resource Include="Resources\Icon.ico" />
-    <Content Include="ReleaseNotes.txt">
+    <Content Include="..\ReleaseNotes.txt">
+	  <Visible>false</Visible>
       <CopyToOutputDirectory>Always</CopyToOutputDirectory>
       <IncludeInVSIX>true</IncludeInVSIX>
     </Content>
