UnrealVS: Fix exception on startup due to missing configuration in launch settings.

diff --git a/Engine/Source/Programs/UnrealVS/UnrealVS.Shared/CommandLineEditor.cs b/Engine/Source/Programs/UnrealVS/UnrealVS.Shared/CommandLineEditor.cs
--- a/Engine/Source/Programs/UnrealVS/UnrealVS.Shared/CommandLineEditor.cs
+++ b/Engine/Source/Programs/UnrealVS/UnrealVS.Shared/CommandLineEditor.cs
@@ -219,15 +219,12 @@ namespace UnrealVS
 										PropertyNameCaseInsensitive = true
 									});
 
-									if (!string.IsNullOrEmpty(activeDebugProfile))
+									LaunchSettingsJson.LaunchSettingsProfile Profile;
+									if (string.IsNullOrEmpty(activeDebugProfile) || !settings.Profiles.TryGetValue(activeDebugProfile, out Profile))
 									{
-										Text = settings.Profiles[activeDebugProfile].CommandLineArgs ?? "";
-									}
-									else
-									{
-										// if no active debug profile is set then VS uses the first one
-										Text = settings.Profiles.FirstOrDefault().Value?.CommandLineArgs ?? "";
+										Profile = settings.Profiles.FirstOrDefault().Value;
 									}
+									Text = Profile?.CommandLineArgs ?? "";
 								}
 							}
 							// for "Game" projects automatically remove the game project filename from the start of the command line
