UnrealVS: Fix thread safety warnings when building.

Mainly done by naively adding ThreadUtils.ThrowIfNotOnUIThread() in places where warnings occur and code looks likely to be on UI thread, then tested a few basic operations in the experimental instance. Remaining warnings look more problematic to solve; I will test these more closely and fix them in a future change.

diff --git a/Engine/Source/Programs/UnrealVS/ReleaseNotes.txt b/Engine/Source/Programs/UnrealVS/ReleaseNotes.txt
--- a/Engine/Source/Programs/UnrealVS/ReleaseNotes.txt
+++ b/Engine/Source/Programs/UnrealVS/ReleaseNotes.txt
@@ -1,6 +1,9 @@
+1.71 (WIP)
+* fixed exception when opening project missing a referenced configuration
+* added exceptions to guard against code UI code not being accessed on the UI thread
+
 1.70
 * add option to allow async checkout
 * resolved concurrent command issues
 * fix issue with P4V commands waiting for exit
 * resolve issue where Checkout on Edit was not honoring the user option
-
diff --git a/Engine/Source/Programs/UnrealVS/UnrealVS.2019/UnrealVS.2019.csproj b/Engine/Source/Programs/UnrealVS/UnrealVS.2019/UnrealVS.2019.csproj
--- a/Engine/Source/Programs/UnrealVS/UnrealVS.2019/UnrealVS.2019.csproj
+++ b/Engine/Source/Programs/UnrealVS/UnrealVS.2019/UnrealVS.2019.csproj
@@ -97,7 +97,7 @@
     <Resource Include="Resources\BatchBuilderBackground.png" />
     <Resource Include="Resources\Icon.ico" />
     <Content Include="..\ReleaseNotes.txt">
-	  <Visible>false</Visible>
+      <Visible>false</Visible>
       <CopyToOutputDirectory>Always</CopyToOutputDirectory>
       <IncludeInVSIX>true</IncludeInVSIX>
     </Content>
@@ -179,7 +179,7 @@
     <PlatformTarget>AnyCPU</PlatformTarget>
     <ErrorReport>prompt</ErrorReport>
     <CodeAnalysisRuleSet>MinimumRecommendedRules.ruleset</CodeAnalysisRuleSet>
-    <NoWarn>MSB3277</NoWarn>
+    <NoWarn>MSB3277;NU1603</NoWarn>
   </PropertyGroup>
   <PropertyGroup Condition=" '$(Configuration)|$(Platform)' == 'Release|AnyCPU' ">
     <OutputPath>bin\Release\</OutputPath>
diff --git a/Engine/Source/Programs/UnrealVS/UnrealVS.Shared/BatchBuilder.cs b/Engine/Source/Programs/UnrealVS/UnrealVS.Shared/BatchBuilder.cs
--- a/Engine/Source/Programs/UnrealVS/UnrealVS.Shared/BatchBuilder.cs
+++ b/Engine/Source/Programs/UnrealVS/UnrealVS.Shared/BatchBuilder.cs
@@ -72,11 +72,14 @@ namespace UnrealVS
 		/// </summary>
 		public void Tick()
 		{
+			ThreadHelper.ThrowIfNotOnUIThread();
 			ToolControl.Tick();
 		}
 
 		private void ShowToolWindow(object sender, EventArgs e)
 		{
+			ThreadHelper.ThrowIfNotOnUIThread();
+
 			// Get the instance number 0 of this tool window. This window is single instance so this instance
 			// is actually the only one.
 			// The last flag is set to true so that if the tool window does not exists it will be created.
diff --git a/Engine/Source/Programs/UnrealVS/UnrealVS.Shared/BatchBuilderToolControl.xaml.cs b/Engine/Source/Programs/UnrealVS/UnrealVS.Shared/BatchBuilderToolControl.xaml.cs
--- a/Engine/Source/Programs/UnrealVS/UnrealVS.Shared/BatchBuilderToolControl.xaml.cs
+++ b/Engine/Source/Programs/UnrealVS/UnrealVS.Shared/BatchBuilderToolControl.xaml.cs
@@ -228,6 +228,7 @@ namespace UnrealVS
 
 			public override string ToString()
 			{
+				ThreadHelper.ThrowIfNotOnUIThread();
 				if (Project != null) return Project.Name;
 				return "<None>";
 			}
@@ -482,6 +483,8 @@ namespace UnrealVS
 
 		private static void DisplayBatchOutputText(string Text)
 		{
+			ThreadHelper.ThrowIfNotOnUIThread();
+
 			if (string.IsNullOrEmpty(Text)) return;
 
 			UnrealVSPackage.Instance.DTE.ExecuteCommand("View.Output");
@@ -506,6 +509,8 @@ namespace UnrealVS
 		/// </summary>
 		public BatchBuilderToolControl(BatchBuilderToolState state)
 		{
+			ThreadHelper.ThrowIfNotOnUIThread();
+
 			_state = state;
 			IsSolutionOpen = UnrealVSPackage.Instance.DTE.Solution.IsOpen;
 
@@ -552,6 +557,7 @@ namespace UnrealVS
 		/// </summary>
 		public void Tick()
 		{
+			ThreadHelper.ThrowIfNotOnUIThread();
 			TickBuildQueue();
 		}
 
@@ -587,6 +593,8 @@ namespace UnrealVS
 
 		private void RefreshConfigAndPlatformCollections()
 		{
+			ThreadHelper.ThrowIfNotOnUIThread();
+
 			_ConfigCollection.Clear();
 			_PlatformCollection.Clear();
 
@@ -611,6 +619,8 @@ namespace UnrealVS
 
 		private void AddButtonClick(object Sender, RoutedEventArgs E)
 		{
+			ThreadHelper.ThrowIfNotOnUIThread();
+
 			BuildJob[] Buildables = GetSelectedBuildableOnLeft();
 
 			if (Buildables.Length > 0)
@@ -663,6 +673,8 @@ namespace UnrealVS
 
 		private BuildJob[] GetSelectedBuildableOnLeft()
 		{
+			ThreadHelper.ThrowIfNotOnUIThread();
+
 			BuildJob.BuildJobType JobType = GetSelectedBuildJobType();
 
 			List<BuildJob> Buildables = new List<BuildJob>();
@@ -719,6 +731,7 @@ namespace UnrealVS
 
 		private void OnBuildDone(bool bSucceeded, bool bModified, bool bWasCancelled)
 		{
+			ThreadHelper.ThrowIfNotOnUIThread();
 			if (_ActiveBuildJob != null)
 			{
 				if (bWasCancelled)
@@ -804,6 +817,7 @@ namespace UnrealVS
 
 		private void OnStartStopButtonClick(object sender, RoutedEventArgs e)
 		{
+			ThreadHelper.ThrowIfNotOnUIThread();
 			if (!IsBusy)
 			{
 				// Start
@@ -831,6 +845,8 @@ namespace UnrealVS
 
 		private void RunBuildJobs()
 		{
+			ThreadHelper.ThrowIfNotOnUIThread();
+
 			if (BuildJobs == null)
 				throw new ArgumentNullException(nameof(BuildJobs));
 
@@ -878,6 +894,8 @@ namespace UnrealVS
 
 		private void SanityCheckBuildJobs(Project[] CheckProjects)
 		{
+			ThreadHelper.ThrowIfNotOnUIThread();
+
 			// check that project, config and platform strings in build jobs are valid
 			// remove invalid ones
 			Utils.GetSolutionConfigsAndPlatforms(out string[] SolutionConfigs, out string[] SolutionPlatforms);
@@ -888,7 +906,9 @@ namespace UnrealVS
 				{
 					BuildJob Job = JobSet.BuildJobs[JobIdx];
 
+#pragma warning disable VSTHRD010 // Invoke single-threaded types on Main thread
 					bool bProjectOkay = CheckProjects.Any(Proj => String.CompareOrdinal(Job.Project.FullName, Proj.FullName) == 0);
+#pragma warning restore VSTHRD010 // Invoke single-threaded types on Main thread
 					bool bConfigOkay = SolutionConfigs.Any(Config => String.CompareOrdinal(Job.Config, Config) == 0);
 					bool bPlatformOkay = SolutionPlatforms.Any(Platform => String.CompareOrdinal(Job.Platform, Platform) == 0);
 
@@ -902,11 +922,14 @@ namespace UnrealVS
 
 		private void SanityCheckBuildJobs()
 		{
+			ThreadHelper.ThrowIfNotOnUIThread();
 			SanityCheckBuildJobs(Utils.GetAllProjectsFromDTE());
 		}
 
 		private void TickBuildQueue()
 		{
+			ThreadHelper.ThrowIfNotOnUIThread();
+
 			if (_ActiveBuildJob != null || _BuildQueue.Count == 0) return;
 
 			while (_ActiveBuildJob == null && _BuildQueue.Count > 0)
@@ -942,6 +965,8 @@ namespace UnrealVS
 
 		private string GetBuildJobOutputText(BuildJob Job, DateTime? StartTime)
 		{
+			ThreadHelper.ThrowIfNotOnUIThread();
+
 			var OutputText = new StringBuilder();
 
 			var Window = UnrealVSPackage.Instance.DTE.Windows.Item(Constants.vsWindowKindOutput);
@@ -992,6 +1017,8 @@ namespace UnrealVS
 
 		private void OpenItemOutputCommandExecuted(object sender, ExecutedRoutedEventArgs e)
 		{
+			ThreadHelper.ThrowIfNotOnUIThread();
+
 			if (!(e.Parameter is BuildJob Job)) return;
 
 			DisplayBatchOutputText(Job.OutputText);
@@ -999,6 +1026,8 @@ namespace UnrealVS
 
 		private void OnDblClickBuildListItem(object sender, MouseButtonEventArgs e)
 		{
+			ThreadHelper.ThrowIfNotOnUIThread();
+
 			if (!(sender is FrameworkElement Elem)) return;
 
 			if (!(Elem.DataContext is BuildJob Job)) return;
@@ -1026,6 +1055,8 @@ namespace UnrealVS
 
 		private void OnDblClickBuildingListItem(object sender, MouseButtonEventArgs e)
 		{
+			ThreadHelper.ThrowIfNotOnUIThread();
+
 			if (!(sender is FrameworkElement Elem)) return;
 
 			if (!(Elem.DataContext is BuildJob Job)) return;
diff --git a/Engine/Source/Programs/UnrealVS/UnrealVS.Shared/BuildStartupProject.cs b/Engine/Source/Programs/UnrealVS/UnrealVS.Shared/BuildStartupProject.cs
--- a/Engine/Source/Programs/UnrealVS/UnrealVS.Shared/BuildStartupProject.cs
+++ b/Engine/Source/Programs/UnrealVS/UnrealVS.Shared/BuildStartupProject.cs
@@ -1,6 +1,7 @@
 // Copyright Epic Games, Inc. All Rights Reserved.
 
 using EnvDTE80;
+using Microsoft.VisualStudio.Shell;
 using Microsoft.VisualStudio.Shell.Interop;
 using System;
 using System.ComponentModel.Design;
@@ -26,6 +27,8 @@ namespace UnrealVS
 		/// Called when 'BuildStartupProject' button is clicked
 		void BuildStartupProjectButtonHandler(object Sender, EventArgs Args)
 		{
+			ThreadHelper.ThrowIfNotOnUIThread();
+
 			// Grab the current startup project
 			UnrealVSPackage.Instance.SolutionBuildManager.get_StartupProject(out IVsHierarchy ProjectHierarchy);
 			if (ProjectHierarchy != null)
diff --git a/Engine/Source/Programs/UnrealVS/UnrealVS.Shared/CommandLineEditor.cs b/Engine/Source/Programs/UnrealVS/UnrealVS.Shared/CommandLineEditor.cs
--- a/Engine/Source/Programs/UnrealVS/UnrealVS.Shared/CommandLineEditor.cs
+++ b/Engine/Source/Programs/UnrealVS/UnrealVS.Shared/CommandLineEditor.cs
@@ -5,6 +5,7 @@ using EnvDTE80;
 using Microsoft.VisualStudio;
 using Microsoft.VisualStudio.Shell;
 using Microsoft.VisualStudio.Shell.Interop;
+using Microsoft.VisualStudio.Threading;
 using System;
 using System.Collections.Generic;
 using System.ComponentModel.Design;
@@ -13,6 +14,7 @@ using System.Linq;
 using System.Runtime.InteropServices;
 using System.Text.Json;
 using System.Text.Json.Serialization;
+using System.Threading.Tasks;
 using System.Timers;
 
 namespace UnrealVS
@@ -63,7 +65,11 @@ namespace UnrealVS
 			// Register for events that we care about
 			UnrealVSPackage.Instance.OnSolutionOpened += UpdateCommandLineCombo;
 			UnrealVSPackage.Instance.OnSolutionClosed += UpdateCommandLineCombo;
-			UnrealVSPackage.Instance.OnStartupProjectChanged += (p) => UpdateCommandLineCombo();
+			UnrealVSPackage.Instance.OnStartupProjectChanged += (p) =>
+			{
+				ThreadHelper.ThrowIfNotOnUIThread();
+				UpdateCommandLineCombo();
+			};
 			UnrealVSPackage.Instance.OnStartupProjectPropertyChanged += OnStartupProjectPropertyChanged;
 			UnrealVSPackage.Instance.OnStartupProjectConfigChanged += OnStartupProjectConfigChanged;
 
@@ -75,6 +81,8 @@ namespace UnrealVS
 		/// </summary>
 		public void OnStartupProjectPropertyChanged(UInt32 itemid, Int32 propid, UInt32 flags)
 		{
+			ThreadHelper.ThrowIfNotOnUIThread();
+
 			// Filter out Helix VS plugin sending thousands of VSHPROPID_StateIconIndex in the solution
 			// This event type is not usefull at all for the commandline editor
 			if (propid == (Int32)__VSHPROPID.VSHPROPID_StateIconIndex)
@@ -148,6 +156,7 @@ namespace UnrealVS
 		/// </summary>
 		private void OnStartupProjectConfigChanged(Project Project)
 		{
+			ThreadHelper.ThrowIfNotOnUIThread();
 			UpdateCommandLineCombo();
 		}
 
@@ -156,6 +165,8 @@ namespace UnrealVS
 		/// </summary>
 		private void UpdateCommandLineCombo()
 		{
+			ThreadHelper.ThrowIfNotOnUIThread();
+
 			// Enable or disable our command-line selector
 			DesiredCommandLine = null;  // clear this state var used by the combo box handler
 			UnrealVSPackage.Instance.SolutionBuildManager.get_StartupProject(out IVsHierarchy ProjectHierarchy);
@@ -181,6 +192,8 @@ namespace UnrealVS
 		/// <returns>String to display</returns>
 		private static string MakeCommandLineComboText()
 		{
+			ThreadHelper.ThrowIfNotOnUIThread();
+
 			string Text = "";
 
 			if (UnrealVSPackage.Instance.SolutionBuildManager.get_StartupProject(out IVsHierarchy ProjectHierarchy) == VSConstants.S_OK && ProjectHierarchy != null)
@@ -252,6 +265,7 @@ namespace UnrealVS
 		/// Called by combo control to query the text to display or to apply newly-entered text
 		private void ComboHandler(object Sender, EventArgs Args)
 		{
+			ThreadHelper.ThrowIfNotOnUIThread();
 			try
 			{
 				var OleArgs = (OleMenuCmdEventArgs)Args;
@@ -307,6 +321,8 @@ namespace UnrealVS
 
 		private void CommitCommandLineText(string CommandLine)
 		{
+			ThreadHelper.ThrowIfNotOnUIThread();
+
 			if (UnrealVSPackage.Instance.SolutionBuildManager.get_StartupProject(out IVsHierarchy ProjectHierarchy) == VSConstants.S_OK && ProjectHierarchy != null)
 			{
 				Project SelectedStartupProject = Utils.HierarchyObjectToProject(ProjectHierarchy);
diff --git a/Engine/Source/Programs/UnrealVS/UnrealVS.Shared/CompileSingleFile.cs b/Engine/Source/Programs/UnrealVS/UnrealVS.Shared/CompileSingleFile.cs
--- a/Engine/Source/Programs/UnrealVS/UnrealVS.Shared/CompileSingleFile.cs
+++ b/Engine/Source/Programs/UnrealVS/UnrealVS.Shared/CompileSingleFile.cs
@@ -48,6 +48,8 @@ namespace UnrealVS
 
 		void CompileSingleFileButtonHandler(object Sender, EventArgs Args)
 		{
+			ThreadHelper.ThrowIfNotOnUIThread();
+
 			MenuCommand SenderSubMenuCommand = (MenuCommand)Sender;
 
 			bool PreprocessOnly = SenderSubMenuCommand.CommandID.ID == PreprocessSingleFileButtonID;
@@ -75,6 +77,7 @@ namespace UnrealVS
 
 		bool TryCompileSingleFile(bool bPreProcessOnly)
 		{
+			ThreadHelper.ThrowIfNotOnUIThread();
 			DTE DTE = UnrealVSPackage.Instance.DTE;
 
 			// Activate the output window
diff --git a/Engine/Source/Programs/UnrealVS/UnrealVS.Shared/GenerateProjectFiles.cs b/Engine/Source/Programs/UnrealVS/UnrealVS.Shared/GenerateProjectFiles.cs
--- a/Engine/Source/Programs/UnrealVS/UnrealVS.Shared/GenerateProjectFiles.cs
+++ b/Engine/Source/Programs/UnrealVS/UnrealVS.Shared/GenerateProjectFiles.cs
@@ -1,5 +1,6 @@
 // Copyright Epic Games, Inc. All Rights Reserved.
 
+using Microsoft.VisualStudio.Shell;
 using System;
 using System.ComponentModel.Design;
 using System.Diagnostics;
@@ -35,6 +36,8 @@ namespace UnrealVS
 		/// Called when 'GenerateProjectFiles' button is clicked
 		void GenerateProjectFilesButtonHandler(object Sender, EventArgs Args)
 		{
+			ThreadHelper.ThrowIfNotOnUIThread();
+
 			// Don't allow a regen if we're already generating projects
 			if (GenerateProjectFilesProcess == null)
 			{
@@ -128,6 +131,8 @@ namespace UnrealVS
 		/// </summary>
 		public void OnOutputFromGenerateProjectFilesProcess(object _Sender, DataReceivedEventArgs Args)
 		{
+			ThreadHelper.ThrowIfNotOnUIThread();
+
 			var Pane = UnrealVSPackage.Instance.GetOutputPane();
 			if (Pane != null)
 			{
diff --git a/Engine/Source/Programs/UnrealVS/UnrealVS.Shared/P4Commands.cs b/Engine/Source/Programs/UnrealVS/UnrealVS.Shared/P4Commands.cs
--- a/Engine/Source/Programs/UnrealVS/UnrealVS.Shared/P4Commands.cs
+++ b/Engine/Source/Programs/UnrealVS/UnrealVS.Shared/P4Commands.cs
@@ -90,6 +90,8 @@ namespace UnrealVS
 
 		private bool IsSolutionLoaded()
 		{
+			ThreadHelper.ThrowIfNotOnUIThread();
+
 			DTE DTE = UnrealVSPackage.Instance.DTE;
 
 			return DTE.Solution.FileName.Length > 0;
@@ -97,6 +99,8 @@ namespace UnrealVS
 
 		private void OnQuickBuildSubMenuQuery(object sender, EventArgs e)
 		{
+			ThreadHelper.ThrowIfNotOnUIThread();
+
 			DTE DTE = UnrealVSPackage.Instance.DTE;
 
 			bool bEnableCommands = UnrealVSPackage.Instance.OptionsPage.AllowUnrealVSP4;
@@ -109,10 +113,12 @@ namespace UnrealVS
 
 		public P4Commands()
 		{
+			ThreadHelper.ThrowIfNotOnUIThread();
+
 			// setup callbacks on IDE operations
 			UnrealVSPackage.Instance.OptionsPage.OnOptionsChanged += OnOptionsChanged;
-			UnrealVSPackage.Instance.OnSolutionOpened += SoltuionOpened;
-			UnrealVSPackage.Instance.OnSolutionClosed += SoltuionClosed;
+			UnrealVSPackage.Instance.OnSolutionOpened += SolutionOpened;
+			UnrealVSPackage.Instance.OnSolutionClosed += SolutionClosed;
 
 			// create specific output window for unrealvs.P4
 			P4OutputPane = UnrealVSPackage.Instance.GetP4OutputPane();
@@ -162,14 +168,16 @@ namespace UnrealVS
 
 		}
 		// Called when solutions are loaded or unloaded
-		private void SoltuionOpened()
+		private void SolutionOpened()
 		{
 			// Update the menu visibility
 			UpdateMenuOptions();
 		}
 
-		private void SoltuionClosed()
+		private void SolutionClosed()
 		{
+			ThreadHelper.ThrowIfNotOnUIThread();
+
 			// Update the menu visibility
 			UpdateMenuOptions();
 
@@ -180,6 +188,8 @@ namespace UnrealVS
 
 		void RegisterCallbackHandler(string CommandName, _dispCommandEvents_BeforeExecuteEventHandler Callback)
 		{
+			ThreadHelper.ThrowIfNotOnUIThread();
+
 			// Find the command from the passed in name
 			DTE DTE = UnrealVSPackage.Instance.DTE;
 			CommandEvents Event = null;
@@ -208,6 +218,8 @@ namespace UnrealVS
 
 		private void UpdateMenuOptions()
 		{
+			ThreadHelper.ThrowIfNotOnUIThread();
+
 			DTE DTE = UnrealVSPackage.Instance.DTE;
 
 			bool bEnableCommands = UnrealVSPackage.Instance.OptionsPage.AllowUnrealVSP4;
@@ -222,6 +234,7 @@ namespace UnrealVS
 		}
 		private void OnOptionsChanged(object Sender, EventArgs E)
 		{
+			ThreadHelper.ThrowIfNotOnUIThread();
 			UpdateMenuOptions();
 		}
 
@@ -245,6 +258,8 @@ namespace UnrealVS
 
 		private void P4CheckoutButtonHandler(object Sender, EventArgs Args)
 		{
+			ThreadHelper.ThrowIfNotOnUIThread();
+
 			DTE DTE = UnrealVSPackage.Instance.DTE;
 
 			// Check we've got a file open
@@ -262,6 +277,8 @@ namespace UnrealVS
 		}
 		private void P4AnnotateButtonHandler(object Sender, EventArgs Args)
 		{
+			ThreadHelper.ThrowIfNotOnUIThread();
+
 			// Debug.ListCallStack
 
 			DTE DTE = UnrealVSPackage.Instance.DTE;
@@ -359,6 +376,8 @@ namespace UnrealVS
 		}
 		private void P4ViewSelectedCLButtonHandler(object Sender, EventArgs Args)
 		{
+			ThreadHelper.ThrowIfNotOnUIThread();
+
 			DTE DTE = UnrealVSPackage.Instance.DTE;
 
 			if (DTE.ActiveDocument == null)
@@ -389,6 +408,8 @@ namespace UnrealVS
 
 		private void P4IntegrationAwareTimeLapseHandler(object Sender, EventArgs Args)
 		{
+			ThreadHelper.ThrowIfNotOnUIThread();
+
 			DTE DTE = UnrealVSPackage.Instance.DTE;
 
 			if (DTE.ActiveDocument == null)
@@ -419,6 +440,8 @@ namespace UnrealVS
 
 		private void P4DiffinVSHandler(object Sender, EventArgs Args)
 		{
+			ThreadHelper.ThrowIfNotOnUIThread();
+
 			DTE DTE = UnrealVSPackage.Instance.DTE;
 
 			if (DTE.ActiveDocument == null)
@@ -471,6 +494,8 @@ namespace UnrealVS
 
 		private void P4GetLast10ChangesHandler(object Sender, EventArgs Args)
 		{
+			ThreadHelper.ThrowIfNotOnUIThread();
+
 			DTE DTE = UnrealVSPackage.Instance.DTE;
 
 			string TempPath = Path.GetTempPath();
@@ -492,6 +517,8 @@ namespace UnrealVS
 
 		private void P4ShowFileInP4Vandler(object Sender, EventArgs Args)
 		{
+			ThreadHelper.ThrowIfNotOnUIThread();
+
 			DTE DTE = UnrealVSPackage.Instance.DTE;
 
 			if (DTE.ActiveDocument == null)
@@ -600,6 +627,8 @@ namespace UnrealVS
 
 		private void PullWorkingDirectory(bool bPullfromP4Settings)
 		{
+			ThreadHelper.ThrowIfNotOnUIThread();
+
 			bPullWorkingDirectorFromP4Mutex.WaitOne();
 
 			if (IsSolutionLoaded() && (P4WorkingDirectory == null || P4WorkingDirectory.Length < 2))
diff --git a/Engine/Source/Programs/UnrealVS/UnrealVS.Shared/QuickBuild.cs b/Engine/Source/Programs/UnrealVS/UnrealVS.Shared/QuickBuild.cs
--- a/Engine/Source/Programs/UnrealVS/UnrealVS.Shared/QuickBuild.cs
+++ b/Engine/Source/Programs/UnrealVS/UnrealVS.Shared/QuickBuild.cs
@@ -104,6 +104,7 @@ namespace UnrealVS
 		/// </summary>
 		private void OnInvokedDynamicItem(object sender, EventArgs args)
 		{
+			ThreadHelper.ThrowIfNotOnUIThread();
 			var MenuCommand = (QuickBuildMenuCommand)sender;
 
 			// Get the project clicked in the solution explorer by accessing the current selection and converting to a Project if possible.
@@ -232,6 +233,8 @@ namespace UnrealVS
 
 		public QuickBuild()
 		{
+			ThreadHelper.ThrowIfNotOnUIThread();
+
 			// root menu
 			QuickBuildCommand = new OleMenuCommand(null, null, OnQuickBuildQuery, new CommandID(GuidList.UnrealVSCmdSet, ProjectQuickBuildMenuID));
 			QuickBuildCommand.BeforeQueryStatus += OnQuickBuildQuery;
@@ -263,6 +266,7 @@ namespace UnrealVS
 		/// </summary>
 		private void UpdateActiveState()
 		{
+			ThreadHelper.ThrowIfNotOnUIThread();
 			UnrealVSPackage.Instance.SelectionManager.IsCmdUIContextActive(SolutionBuildingUIContextCookie, out int bIsBuilding);
 			bIsActive = bIsBuilding == 0;
 		}
@@ -272,6 +276,8 @@ namespace UnrealVS
 		/// </summary>
 		private void OnQuickBuildQuery(object sender, EventArgs e)
 		{
+			ThreadHelper.ThrowIfNotOnUIThread();
+
 			// Always cache the list of solution build configs when the project menu is opening
 			CacheBuildConfigs();
 		}
@@ -292,6 +298,8 @@ namespace UnrealVS
 		/// </summary>
 		private void OnSolutionChanged()
 		{
+			ThreadHelper.ThrowIfNotOnUIThread();
+
 			CacheBuildConfigs();
 
 			ActiveConfigurationMenus.Clear();
@@ -310,12 +318,13 @@ namespace UnrealVS
 		/// </summary>
 		private void CacheBuildConfigs()
 		{
+			ThreadHelper.ThrowIfNotOnUIThread();
+
 			SolutionConfigurations SolutionConfigs = UnrealVSPackage.Instance.DTE.Solution.SolutionBuild.SolutionConfigurations;
 
 			SolutionConfigPlatforms = (from SolutionConfiguration2 Sc in SolutionConfigs select Sc.PlatformName).Distinct().ToArray();
 
 			SolutionConfigNames = (from SolutionConfiguration2 Sc in SolutionConfigs select Sc.Name).Distinct().ToArray();
-
 		}
 
 		/// <summary>
@@ -323,6 +332,7 @@ namespace UnrealVS
 		/// </summary>
 		private void OnUIContextChanged(uint CmdUICookie, bool bActive)
 		{
+			ThreadHelper.ThrowIfNotOnUIThread();
 			if (SolutionBuildingUIContextCookie == CmdUICookie)
 			{
 				UpdateActiveState();
diff --git a/Engine/Source/Programs/UnrealVS/UnrealVS.Shared/StartupProjectSelector.cs b/Engine/Source/Programs/UnrealVS/UnrealVS.Shared/StartupProjectSelector.cs
--- a/Engine/Source/Programs/UnrealVS/UnrealVS.Shared/StartupProjectSelector.cs
+++ b/Engine/Source/Programs/UnrealVS/UnrealVS.Shared/StartupProjectSelector.cs
@@ -27,21 +27,22 @@ namespace UnrealVS
 
 		private class ProjectReference
 		{
-			public Project Project { get; set; }
+			public Project Project { get; }
+			public string Name { get; }
 
-			public string Name
+			public ProjectReference(Project Project)
 			{
-				get { return (Project != null ? Project.Name : "<invalid>"); }
+				ThreadHelper.ThrowIfNotOnUIThread();
+				this.Project = Project;
+				this.Name = (Project != null) ? Project.Name : "<invalid>";
 			}
 
-			public override string ToString()
-			{
-				return (Project != null ? Project.Name : "<invalid>");
-			}
+			public override string ToString() => Name;
 		}
 
 		public StartupProjectSelector()
 		{
+			ThreadHelper.ThrowIfNotOnUIThread();
 			_bIsSolutionOpened = UnrealVSPackage.Instance.DTE.Solution.IsOpen;
 
 			// Create the handlers for our commands
@@ -66,6 +67,7 @@ namespace UnrealVS
 			UnrealVSPackage.Instance.OnSolutionOpened +=
 				delegate
 				{
+					ThreadHelper.ThrowIfNotOnUIThread();
 					Logging.WriteLine("Opened solution " + UnrealVSPackage.Instance.DTE.Solution.FullName);
 					_bIsSolutionOpened = true;
 					UpdateStartupProjectCombo();
@@ -82,6 +84,7 @@ namespace UnrealVS
 			UnrealVSPackage.Instance.OnProjectOpened +=
 				delegate (Project OpenedProject)
 				{
+					ThreadHelper.ThrowIfNotOnUIThread();
 					if (_bIsSolutionOpened)
 					{
 						Logging.WriteLine("Opened project node " + OpenedProject.Name);
@@ -95,6 +98,7 @@ namespace UnrealVS
 			UnrealVSPackage.Instance.OnProjectClosed +=
 				delegate (Project ClosedProject)
 				{
+					ThreadHelper.ThrowIfNotOnUIThread();
 					Logging.WriteLine("Closed project node " + ClosedProject.Name);
 					RemoveFromStartupProjectList(ClosedProject);
 				};
@@ -114,6 +118,8 @@ namespace UnrealVS
 
 		private void OnOptionsChanged(object Sender, EventArgs E)
 		{
+			ThreadHelper.ThrowIfNotOnUIThread();
+
 			UpdateCachedOptions();
 
 			UpdateStartupProjectList(Utils.GetAllProjectsFromDTE());
@@ -127,13 +133,15 @@ namespace UnrealVS
 		/// Rebuild/update the list CachedStartupProjectNames whenever something changes
 		private void UpdateStartupProjectList(params Project[] DirtyProjects)
 		{
+			ThreadHelper.ThrowIfNotOnUIThread();
+
 			foreach (Project Project in DirtyProjects)
 			{
 				if (ProjectPredicate(Project))
 				{
 					if (!_CachedStartupProjects.Any(ProjRef => ProjRef.Project == Project))
 					{
-						_CachedStartupProjects.Add(new ProjectReference { Project = Project });
+						_CachedStartupProjects.Add(new ProjectReference(Project));
 					}
 				}
 				else
@@ -149,6 +157,8 @@ namespace UnrealVS
 
 		private void SortCachedStartupProjects()
 		{
+			ThreadHelper.ThrowIfNotOnUIThread();
+
 			// Sort projects by game y/n then alphabetically
 			_CachedStartupProjects = (_CachedStartupProjects.OrderBy(ProjectRef => ProjectRef.Name == "UE5" ? 0 : 1)
 				.ThenBy(ProjectRef => Utils.IsGameProject(ProjectRef.Project) ? 0 : 1)
@@ -157,6 +167,8 @@ namespace UnrealVS
 
 		private bool ProjectPredicate(Project Project)
 		{
+			ThreadHelper.ThrowIfNotOnUIThread();
+
 			// Initialize the cached options if they haven't been read from UnrealVSOptions yet
 			if (!CachedHideNonGameStartupProjects.HasValue)
 			{
@@ -219,6 +231,8 @@ namespace UnrealVS
 		/// <returns>String to display</returns>
 		private string MakeStartupProjectComboText()
 		{
+			ThreadHelper.ThrowIfNotOnUIThread();
+
 			string Text = "";
 
 			// Switch to this project!
@@ -239,6 +253,8 @@ namespace UnrealVS
 		/// Called by combo control to query the text to display or to apply newly-entered text
 		void StartupProjectComboHandler(object Sender, EventArgs Args)
 		{
+			ThreadHelper.ThrowIfNotOnUIThread();
+
 			var OleArgs = (OleMenuCmdEventArgs)Args;
 
 			// If OutValue is non-zero, then Visual Studio is querying the current combo value to display
@@ -275,6 +291,8 @@ namespace UnrealVS
 		/// Called by combo control to populate the drop-down list
 		void StartupProjectComboListHandler(object Sender, EventArgs Args)
 		{
+			ThreadHelper.ThrowIfNotOnUIThread();
+
 			var OleArgs = (OleMenuCmdEventArgs)Args;
 
 			var CachedStartupProjectNames = (from ProjRef in _CachedStartupProjects select ProjRef.Name).ToArray();
diff --git a/Engine/Source/Programs/UnrealVS/UnrealVS.Shared/UnrealVSPackage.cs b/Engine/Source/Programs/UnrealVS/UnrealVS.Shared/UnrealVSPackage.cs
--- a/Engine/Source/Programs/UnrealVS/UnrealVS.Shared/UnrealVSPackage.cs
+++ b/Engine/Source/Programs/UnrealVS/UnrealVS.Shared/UnrealVSPackage.cs
@@ -17,22 +17,13 @@ using System.Runtime.InteropServices;
 using System.Threading;
 using Microsoft.VisualStudio.TextManager.Interop;
 using Thread = System.Threading.Thread;
-using System;
-using System.Runtime.InteropServices;
 using System.Windows.Forms;
 using Microsoft.VisualStudio.ComponentModelHost;
-using Microsoft.VisualStudio.Editor;
-using Microsoft.VisualStudio.OLE.Interop;
-using Microsoft.VisualStudio.Shell.Interop;
-using Microsoft.VisualStudio.Shell;
 using Microsoft.VisualStudio.Text;
 using Microsoft.VisualStudio.Text.Editor;
-using Microsoft.VisualStudio.TextManager.Interop;
 using IOleServiceProvider = Microsoft.VisualStudio.OLE.Interop.IServiceProvider;
 using Microsoft.VisualStudio.Utilities;
 using Microsoft.VisualStudio.Text.Projection;
-using Microsoft.VisualStudio;
-using System.Collections.Generic;
 
 namespace UnrealVS
 {
@@ -320,6 +311,7 @@ namespace UnrealVS
 		/// </summary>
 		private void Tick()
 		{
+			ThreadHelper.ThrowIfNotOnUIThread();
 			BatchBuilder.Tick();
 		}
 
@@ -328,12 +320,15 @@ namespace UnrealVS
 		/// </summary>
 		public void Dispose()
 		{
+			ThreadHelper.ThrowIfNotOnUIThread();
 			Dispose(true);
 		}
 
 		/// IDispose pattern lets us clean up our stuff!
 		protected override void Dispose(bool disposing)
 		{
+			ThreadHelper.ThrowIfNotOnUIThread();
+
 			if (Ticker != null && Ticker.IsAlive)
 			{
 				Thread.Sleep(TickPeriod + TickPeriod);
@@ -401,6 +396,7 @@ namespace UnrealVS
 		{
 			get
 			{
+				ThreadHelper.ThrowIfNotOnUIThread();
 				if (_DTE2 == null)
 				{
 					// Get the interface when first used.
@@ -607,6 +603,8 @@ namespace UnrealVS
 
 		int IVsSolutionEvents.OnAfterCloseSolution(object pUnkReserved)
 		{
+			ThreadHelper.ThrowIfNotOnUIThread();
+
 			UpdateUnrealLoadedStatus();
 
 			OnSolutionClosed?.Invoke();
@@ -621,6 +619,8 @@ namespace UnrealVS
 
 		int IVsSolutionEvents.OnAfterOpenProject(IVsHierarchy pHierarchy, int fAdded)
 		{
+			ThreadHelper.ThrowIfNotOnUIThread();
+
 			// This function is called after a Visual Studio project is opened (or a new project is created.)
 
 			// Get the actual Project object from the IVsHierarchy object that was supplied
@@ -637,6 +637,8 @@ namespace UnrealVS
 
 		int IVsSolutionEvents.OnAfterOpenSolution(object pUnkReserved, int fNewSolution)
 		{
+			ThreadHelper.ThrowIfNotOnUIThread();
+
 			UpdateUnrealLoadedStatus();
 
 			StartTicker();
@@ -648,6 +650,8 @@ namespace UnrealVS
 
 		int IVsSolutionEvents.OnBeforeCloseProject(IVsHierarchy pHierarchy, int fRemoved)
 		{
+			ThreadHelper.ThrowIfNotOnUIThread();
+
 			// This function is called after a Visual Studio project is closed
 
 			// Get the actual Project object from the IVsHierarchy object that was supplied
@@ -709,6 +713,8 @@ namespace UnrealVS
 
 		int IVsSelectionEvents.OnElementValueChanged(uint elementid, object varValueOld, object varValueNew)
 		{
+			ThreadHelper.ThrowIfNotOnUIThread();
+
 			// This function is called when selection changes in various Visual Studio tool windows
 			// and sub-systems.
 
@@ -828,6 +834,8 @@ namespace UnrealVS
 
 		int IVsPersistSolutionProps.ReadSolutionProps(IVsHierarchy pHierarchy, string pszProjectName, string pszProjectMk, string pszKey, int fPreLoad, IPropertyBag pPropBag)
 		{
+			ThreadHelper.ThrowIfNotOnUIThread();
+
 			if (!string.Equals(pszKey, GuidList.UnrealVSPackageString, StringComparison.InvariantCultureIgnoreCase))
 				return VSConstants.S_OK;
 
@@ -881,6 +889,8 @@ namespace UnrealVS
 
 		public int OnActiveProjectCfgChange(IVsHierarchy pIVsHierarchy)
 		{
+			ThreadHelper.ThrowIfNotOnUIThread();
+
 			// This function is called after a Visual Studio project has its active config changed
 
 			// Check whether the project is the current startup project
@@ -899,6 +909,8 @@ namespace UnrealVS
 
 		private void UpdateUnrealLoadedStatus()
 		{
+			ThreadHelper.ThrowIfNotOnUIThread();
+
 			if (!DTE.Solution.IsOpen)
 			{
 				IsUESolutionLoaded = false;
@@ -987,6 +999,8 @@ namespace UnrealVS
 		/// Obtains the DTE2 interface for this instance of VS from the RunningObjectTable
 		private static DTE2 GetDTE2ForCurrentInstance(DTE DTE)
 		{
+			ThreadHelper.ThrowIfNotOnUIThread();
+
 			// Find the ROT entry for visual studio running under current process.
 			int HResult = NativeMethods.GetRunningObjectTable(0, out IRunningObjectTable Rot);
 			if (HResult == 0)
diff --git a/Engine/Source/Programs/UnrealVS/UnrealVS.Shared/Utils.cs b/Engine/Source/Programs/UnrealVS/UnrealVS.Shared/Utils.cs
--- a/Engine/Source/Programs/UnrealVS/UnrealVS.Shared/Utils.cs
+++ b/Engine/Source/Programs/UnrealVS/UnrealVS.Shared/Utils.cs
@@ -3,6 +3,7 @@
 using EnvDTE;
 using EnvDTE80;
 using Microsoft.VisualStudio;
+using Microsoft.VisualStudio.Shell;
 using Microsoft.VisualStudio.Shell.Interop;
 using System;
 using System.Collections.Generic;
@@ -23,8 +24,12 @@ namespace UnrealVS
 
 			public Project GetProjectSlow()
 			{
+				ThreadHelper.ThrowIfNotOnUIThread();
+
 				Project[] Projects = GetAllProjectsFromDTE();
+#pragma warning disable VSTHRD010 // Invoke single-threaded types on Main thread
 				return Projects.FirstOrDefault(Proj => string.CompareOrdinal(Proj.FullName, FullName) == 0);
+#pragma warning restore VSTHRD010 // Invoke single-threaded types on Main thread
 			}
 		}
 
@@ -35,6 +40,7 @@ namespace UnrealVS
 		/// <returns>IVsHierarchy for the specified project</returns>
 		public static IVsHierarchy ProjectToHierarchyObject(Project Project)
 		{
+			ThreadHelper.ThrowIfNotOnUIThread();
 			UnrealVSPackage.Instance.SolutionManager.GetProjectOfUniqueName(Project.FullName, out IVsHierarchy HierarchyObject);
 			return HierarchyObject;
 		}
@@ -47,6 +53,8 @@ namespace UnrealVS
 		/// <returns>Visual Studio project object</returns>
 		public static Project HierarchyObjectToProject(IVsHierarchy HierarchyObject)
 		{
+			ThreadHelper.ThrowIfNotOnUIThread();
+
 			// Get the actual Project object from the IVsHierarchy object that was supplied
 			HierarchyObject.GetProperty(VSConstants.VSITEMID_ROOT, (int)__VSHPROPID.VSHPROPID_ExtObject, out object ProjectObject);
 			return (Project)ProjectObject;
@@ -59,6 +67,8 @@ namespace UnrealVS
 		/// <returns>Visual Studio project object</returns>
 		public static IVsCfgProvider2 HierarchyObjectToCfgProvider(IVsHierarchy HierarchyObject)
 		{
+			ThreadHelper.ThrowIfNotOnUIThread();
+
 			// Get the actual Project object from the IVsHierarchy object that was supplied
 			HierarchyObject.GetProperty(VSConstants.VSITEMID_ROOT, (int)__VSHPROPID.VSHPROPID_BrowseObject, out object BrowseObject);
 
@@ -78,6 +88,8 @@ namespace UnrealVS
 
 		private static IVsCfgProvider2 GetCfgProviderFromObject(object SomeObject)
 		{
+			ThreadHelper.ThrowIfNotOnUIThread();
+
 			IVsCfgProvider2 CfgProvider2 = null;
 
 			if (SomeObject is IVsGetCfgProvider GetCfgProvider)
@@ -105,6 +117,8 @@ namespace UnrealVS
 		/// <returns>Property object or null if not found</returns>
 		public static Property GetProjectProperty(Project Project, string PropertyName)
 		{
+			ThreadHelper.ThrowIfNotOnUIThread();
+
 			var Properties = Project.Properties;
 			if (Properties != null)
 			{
@@ -129,6 +143,7 @@ namespace UnrealVS
 		/// <param name="PropertyValue">Value to set for this property</param>
 		public static void SetPropertyValue(Property Property, object PropertyValue)
 		{
+			ThreadHelper.ThrowIfNotOnUIThread();
 			Property.Value = PropertyValue;
 
 			// @todo: Not sure if actually needed for command-line property (saved in .user files, not in project)
@@ -146,7 +161,14 @@ namespace UnrealVS
 		{
 			public UIHierarchyItem Item { get; set; }
 			public UITreeItem[] Children { get; set; }
-			public string Name { get { return Item != null ? Item.Name : "None"; } }
+			public string Name
+			{ 
+				get 
+				{
+					ThreadHelper.ThrowIfNotOnUIThread();
+					return Item != null ? Item.Name : "None"; 
+				} 
+			}
 			public object Object { get { return Item?.Object; } }
 		}
 
@@ -155,6 +177,7 @@ namespace UnrealVS
 		/// </summary>
 		public static UITreeItem GetUIHierarchyTree(UIHierarchy Hierarchy)
 		{
+			ThreadHelper.ThrowIfNotOnUIThread();
 			return new UITreeItem
 			{
 				Item = null,
@@ -167,6 +190,7 @@ namespace UnrealVS
 		/// </summary>
 		private static UITreeItem GetUIHierarchyTree(UIHierarchyItem HierarchyItem)
 		{
+			ThreadHelper.ThrowIfNotOnUIThread();
 			return new UITreeItem
 			{
 				Item = HierarchyItem,
@@ -233,6 +257,7 @@ namespace UnrealVS
 		/// </summary>
 		public static bool IsProjectBuildable(Project Project)
 		{
+			ThreadHelper.ThrowIfNotOnUIThread();
 			return Project.Kind == GuidList.VCSharpProjectKindGuidString || Project.Kind == GuidList.VCProjectKindGuidString;
 		}
 
@@ -240,6 +265,8 @@ namespace UnrealVS
 		/// Recurses into items because these are actually in a tree structure
 		public static Project[] GetAllProjectsFromDTE()
 		{
+			ThreadHelper.ThrowIfNotOnUIThread();
+
 			try
 			{
 				List<Project> Projects = new List<Project>();
@@ -274,6 +301,7 @@ namespace UnrealVS
 												Action ExecutingDelegate,
 												Action FailedToStartDelegate)
 		{
+			ThreadHelper.ThrowIfNotOnUIThread();
 			IVsHierarchy ProjHierarchy = Utils.ProjectToHierarchyObject(Project);
 
 			if (ProjHierarchy != null)
@@ -362,6 +390,8 @@ namespace UnrealVS
 
 		private static bool LoadConfigFromUBT(Project SelectedProject)
 		{
+			ThreadHelper.ThrowIfNotOnUIThread();
+
 			string ProjectPath = Path.GetDirectoryName(SelectedProject.FullName);
 			string ConfigFileName = Path.Combine(ProjectPath, "UnrealVS.xml");
 
@@ -390,6 +420,8 @@ namespace UnrealVS
 
 		public static List<string> GetExtraDebuggerCommandArguments(string PlatformName, Project SelectedProject)
 		{
+			ThreadHelper.ThrowIfNotOnUIThread();
+
 			List<string> Result = new List<string>();
 
 			if (LoadConfigFromUBT(SelectedProject))
@@ -412,6 +444,7 @@ namespace UnrealVS
 
 		public static bool IsGameProject(Project Project)
 		{
+			ThreadHelper.ThrowIfNotOnUIThread();
 			return GetUProjects().ContainsKey(Project.Name);
 		}
 
@@ -425,11 +458,14 @@ namespace UnrealVS
 
 		public static string GetUProjectFileName(Project Project)
 		{
+			ThreadHelper.ThrowIfNotOnUIThread();
 			return Project.Name + "." + UProjectExtension;
 		}
 
 		public static string GetAutoUProjectCommandLinePrefix(Project Project)
 		{
+			ThreadHelper.ThrowIfNotOnUIThread();
+
 			var UProjectFileName = GetUProjectFileName(Project);
 			var AllUProjects = GetUProjects();
 
@@ -510,6 +546,8 @@ namespace UnrealVS
 		/// </summary>
 		public static IDictionary<string, string> GetUProjects()
 		{
+			ThreadHelper.ThrowIfNotOnUIThread();
+
 			var Folder = GetSolutionFolder();
 			if (string.IsNullOrEmpty(Folder))
 			{
@@ -566,6 +604,8 @@ namespace UnrealVS
 
 		public static void GetSolutionConfigsAndPlatforms(out string[] SolutionConfigs, out string[] SolutionPlatforms)
 		{
+			ThreadHelper.ThrowIfNotOnUIThread();
+
 			var UniqueConfigs = new List<string>();
 			var UniquePlatforms = new List<string>();
 
@@ -588,6 +628,8 @@ namespace UnrealVS
 
 		public static bool SetActiveSolutionConfiguration(string ConfigName, string PlatformName)
 		{
+			ThreadHelper.ThrowIfNotOnUIThread();
+
 			SolutionConfigurations DteSolutionConfigs = UnrealVSPackage.Instance.DTE.Solution.SolutionBuild.SolutionConfigurations;
 			foreach (SolutionConfiguration2 SolutionConfig in DteSolutionConfigs)
 			{
@@ -603,6 +645,8 @@ namespace UnrealVS
 
 		public static bool SelectProjectInSolutionExplorer(Project Project)
 		{
+			ThreadHelper.ThrowIfNotOnUIThread();
+
 			UnrealVSPackage.Instance.DTE.ExecuteCommand("View.SolutionExplorer");
 			if (Project.ParentProjectItem != null)
 			{
@@ -613,7 +657,9 @@ namespace UnrealVS
 			Utils.UITreeItem SolutionExplorerTree = Utils.GetUIHierarchyTree(SolutionExplorerHierarachy);
 			var UIHierarachyProjects = Utils.GetUITreeItemsByObjectType<Project>(SolutionExplorerTree);
 
+#pragma warning disable VSTHRD010 // Invoke single-threaded types on Main thread
 			var SelectableUIItem = UIHierarachyProjects.FirstOrDefault(uihp => uihp.Object as Project == Project);
+#pragma warning restore VSTHRD010 // Invoke single-threaded types on Main thread
 
 			if (SelectableUIItem != null)
 			{
@@ -634,6 +680,8 @@ namespace UnrealVS
 
 		private static void PrepareOutputPane()
 		{
+			ThreadHelper.ThrowIfNotOnUIThread();
+
 			UnrealVSPackage.Instance.DTE.ExecuteCommand("View.Output");
 
 			var Pane = UnrealVSPackage.Instance.GetOutputPane();
@@ -650,6 +698,8 @@ namespace UnrealVS
 		/// Called by GetAllProjectsFromDTE() to list items from the project tree
 		private static void GetSubProjectsOfProjectItem(ProjectItem Item, List<Project> Projects)
 		{
+			ThreadHelper.ThrowIfNotOnUIThread();
+
 			if (Item.SubProject != null)
 			{
 				Projects.Add(Item.SubProject);
@@ -673,6 +723,8 @@ namespace UnrealVS
 
 		private static string GetSolutionFolder()
 		{
+			ThreadHelper.ThrowIfNotOnUIThread();
+
 			if (!UnrealVSPackage.Instance.DTE.Solution.IsOpen)
 			{
 				return string.Empty;
