[UnrealVS] - add p4FastReconcile

diff --git a/Engine/Source/Programs/UnrealVS/UnrealVS.2019/Properties/AssemblyInfo.cs b/Engine/Source/Programs/UnrealVS/UnrealVS.2019/Properties/AssemblyInfo.cs
--- a/Engine/Source/Programs/UnrealVS/UnrealVS.2019/Properties/AssemblyInfo.cs
+++ b/Engine/Source/Programs/UnrealVS/UnrealVS.2019/Properties/AssemblyInfo.cs
@@ -31,8 +31,8 @@ using System.Runtime.InteropServices;
 // You can specify all the values or you can default the Revision and Build Numbers 
 // by using the '*' as shown below:
 
-[assembly: AssemblyVersion("1.70.0.0")]
-[assembly: AssemblyFileVersion("1.70.0.0")]
+[assembly: AssemblyVersion("1.71.0.0")]
+[assembly: AssemblyFileVersion("1.71.0.0")]
 
 
 
diff --git a/Engine/Source/Programs/UnrealVS/UnrealVS.2022/Properties/AssemblyInfo.cs b/Engine/Source/Programs/UnrealVS/UnrealVS.2022/Properties/AssemblyInfo.cs
--- a/Engine/Source/Programs/UnrealVS/UnrealVS.2022/Properties/AssemblyInfo.cs
+++ b/Engine/Source/Programs/UnrealVS/UnrealVS.2022/Properties/AssemblyInfo.cs
@@ -30,8 +30,8 @@ using System.Runtime.InteropServices;
 // You can specify all the values or you can default the Revision and Build Numbers 
 // by using the '*' as shown below:
 
-[assembly: AssemblyVersion("1.70.0.0")]
-[assembly: AssemblyFileVersion("1.70.0.0")]
+[assembly: AssemblyVersion("1.71.0.0")]
+[assembly: AssemblyFileVersion("1.71.0.0")]
 
 
 
diff --git a/Engine/Source/Programs/UnrealVS/UnrealVS.Shared/P4Commands.cs b/Engine/Source/Programs/UnrealVS/UnrealVS.Shared/P4Commands.cs
--- a/Engine/Source/Programs/UnrealVS/UnrealVS.Shared/P4Commands.cs
+++ b/Engine/Source/Programs/UnrealVS/UnrealVS.Shared/P4Commands.cs
@@ -36,6 +36,7 @@ namespace UnrealVS
 		private const int P4DiffinVSButtonID = 0x1454;
 		private const int P4GetLast10ChangesID = 0x1455;
 		private const int P4ShowFileinP4VID = 0x1456;
+		private const int P4FastReconcileCodeFilesID  =0x1457;
 
 		private OleMenuCommand SubMenuCommand;
 
@@ -62,6 +63,7 @@ namespace UnrealVS
 		private string Username = "";
 		private string Port = "";
 		private string Client = "";
+		private string Stream = "";
 		private string UserInfoComplete = "";
 
 		private List<CommandEvents> EventsForce = new List<CommandEvents>();
@@ -146,6 +148,8 @@ namespace UnrealVS
 			P4CommandsList.Add(new P4Command(P4DiffinVSButtonID, P4DiffinVSHandler));
 			P4CommandsList.Add(new P4Command(P4GetLast10ChangesID, P4GetLast10ChangesHandler));
 			P4CommandsList.Add(new P4Command(P4ShowFileinP4VID, P4ShowFileInP4Vandler));
+			P4CommandsList.Add(new P4Command(P4FastReconcileCodeFilesID, P4FastReconcileCodeFiles));
+			
 
 			if (P4VCCmd.Length > 1)
 			{
@@ -515,6 +519,53 @@ namespace UnrealVS
 			}
 		}
 
+		private void P4FastReconcileCodeFiles(object Sender, EventArgs Args)
+		{
+			DTE DTE = UnrealVSPackage.Instance.DTE;
+
+			string[] DefaultExtensions =
+			{
+				"c*",
+				"h*",
+				"ini"
+			};
+
+			string[] Extensions = UnrealVSPackage.Instance.OptionsPage.ReconcileExtensions.Split(';');
+
+			if (Extensions.Length < 1)
+			{
+				Extensions = DefaultExtensions;
+			}
+
+			// generate the changes list
+			string ReconcileDepotPaths = "";
+			string Preview = "-n";
+
+			if (Stream.Length < 1)
+			{
+				PullWorkingDirectory(bPullWorkingDirectoryOn);
+			}
+
+			foreach (string Ext in Extensions)
+			{
+				ReconcileDepotPaths += Stream + "/...." + Ext + " ";
+			}
+
+			if (UnrealVSPackage.Instance.OptionsPage.AllowReconcileToMarkForEdit)
+			{
+				Preview = "";
+			}
+			else
+			{
+				P4OutputPane.OutputString($"PREVIEW{Environment.NewLine}");
+			}
+
+			P4OutputPane.OutputString($"Running Async Reconcile on : {ReconcileDepotPaths}{Environment.NewLine}");
+			P4OutputPane.Activate();
+
+			TryP4CommandAsync($"reconcile -e -m {Preview} {ReconcileDepotPaths}");
+		}
+
 		private void P4ShowFileInP4Vandler(object Sender, EventArgs Args)
 		{
 			ThreadHelper.ThrowIfNotOnUIThread();
@@ -611,14 +662,17 @@ namespace UnrealVS
 			Regex UserPattern = new Regex(@"User name: (?<user>.*)$", RegexOptions.Compiled | RegexOptions.Multiline);
 			Regex PortPattern = new Regex(@"Server address: (?<port>.*)$", RegexOptions.Compiled | RegexOptions.Multiline);
 			Regex ClientPattern = new Regex(@"Client name: (?<client>.*)$", RegexOptions.Compiled | RegexOptions.Multiline);
+			Regex ClientStream = new Regex(@"Client stream: (?<stream>.*)$", RegexOptions.Compiled | RegexOptions.Multiline);
 
 			System.Text.RegularExpressions.Match UserMatch = UserPattern.Match(CaptureP4Output);
 			System.Text.RegularExpressions.Match PortMatch = PortPattern.Match(CaptureP4Output);
-			System.Text.RegularExpressions.Match ClientMath = ClientPattern.Match(CaptureP4Output);
+			System.Text.RegularExpressions.Match ClientMatch = ClientPattern.Match(CaptureP4Output);
+			System.Text.RegularExpressions.Match StreamMatch = ClientStream.Match(CaptureP4Output);
 
 			Port = PortMatch.Groups["port"].Value.Trim();
 			Username = UserMatch.Groups["user"].Value.Trim();
-			Client = ClientMath.Groups["client"].Value.Trim();
+			Client = ClientMatch.Groups["client"].Value.Trim();
+			Stream = StreamMatch.Groups["stream"].Value.Trim();
 
 			UserInfoComplete = string.Format(" -p {0} -u {1} -c {2} ", Port, Username, Client);
 
@@ -698,26 +752,16 @@ namespace UnrealVS
 
 			DTE DTE = UnrealVSPackage.Instance.DTE;
 
-			// If there's already a operation in progress, abort... potential issues here
-			// in the OnSave/OnSaveAll operations where we will stall.
-/*			if (ChildProcess != null && !ChildProcess.HasExited)
+			// Set up the output pane
+			if (UnrealVSPackage.Instance.OptionsPage.ForceOutputWindow)
 			{
-				P4OutputPane.OutputString($"P4 operation already in flight...waiting{Environment.NewLine}");
 				P4OutputPane.Activate();
-				ChildProcess.WaitForExit();
-			}*/
-
-			// Make sure any existing op is stopped
-			//KillChildProcess();
-
-			// Set up the output pane
-			P4OutputPane.Activate();
-			//P4OutputPane.OutputString($"1>------ P4 Operation started{Environment.NewLine}");
-			//P4OutputPane.OutputString($"CWD: {P4WorkingDirectory}{Environment.NewLine}");
-			//P4OutputPane.OutputString($"CMD: {P4Exe} {CommandLine}{Environment.NewLine}");
-
+			}
+			
+			
 			StringBuilder StdOutSB = new StringBuilder();
 			StringBuilder StdErrSB = new StringBuilder();
+			DateTime Start = DateTime.Now;
 
 			// Spawn the new process
 			System.Diagnostics.Process ChildProcess = new System.Diagnostics.Process()
@@ -754,7 +798,8 @@ namespace UnrealVS
 						P4OutputPane.OutputString(P4OperationStdErr);
 					}
 
-					//P4OutputPane.OutputString($"1>------ P4 Operation complete{Environment.NewLine}");
+					TimeSpan Duration = DateTime.Now - Start;
+					P4OutputPane.OutputString($"complete {Duration.TotalSeconds.ToString()} {Environment.NewLine}");
 
 					ChildProcessList.Remove(ChildProcess);
 				};
@@ -769,6 +814,7 @@ namespace UnrealVS
 			if (bWaitForCompletion)
 			{ 
 				ChildProcess.WaitForExit();
+				TimeSpan Duration = DateTime.Now - Start;
 
 				CaptureP4Output = StdOutSB.ToString();
 				P4OperationStdErr = StdErrSB.ToString();
@@ -783,7 +829,7 @@ namespace UnrealVS
 					P4OutputPane.OutputString(P4OperationStdErr);
 				}
 
-				//P4OutputPane.OutputString($"1>------ P4 Operation complete{Environment.NewLine}");
+				P4OutputPane.OutputString($"complete {((int)Duration.TotalSeconds).ToString()}seconds {Environment.NewLine}");
 
 				ChildProcessList.Remove(ChildProcess);
 
diff --git a/Engine/Source/Programs/UnrealVS/UnrealVS.Shared/ReleaseNotes.txt b/Engine/Source/Programs/UnrealVS/UnrealVS.Shared/ReleaseNotes.txt
new file mode 100644
--- /dev/null
+++ b/Engine/Source/Programs/UnrealVS/UnrealVS.Shared/ReleaseNotes.txt
@@ -0,0 +1,10 @@
+1.71
+* add p4FastReconcile
+* add ability to halt p4 output auto focus
+
+1.70
+* add option to allow async checkout
+* resolved concurrent command issues
+* fix issue with P4V commands waiting for exit
+* resolve issue where Checkout on Edit was not honoring the user option
+
diff --git a/Engine/Source/Programs/UnrealVS/UnrealVS.Shared/UnrealVSOptions.cs b/Engine/Source/Programs/UnrealVS/UnrealVS.Shared/UnrealVSOptions.cs
--- a/Engine/Source/Programs/UnrealVS/UnrealVS.Shared/UnrealVSOptions.cs
+++ b/Engine/Source/Programs/UnrealVS/UnrealVS.Shared/UnrealVSOptions.cs
@@ -36,6 +36,22 @@ namespace UnrealVS
 		[DisplayName("Allow Async Checkout")]
 		[Description("Switches checkout (on edit or manual) to an async method, note if the p4 operation fails the local file will still be writeable")]
 		public bool AllowAsyncP4Checkout { get; set; }
+
+		[Category("Unreal.P4")]
+		[DisplayName("Force P4 Output window to active")]
+		[Description("When set, normal p4 operations will force the output window to the surface")]
+		public bool ForceOutputWindow { get; set; }
+
+		[Category("Unreal.P4")]
+		[DisplayName("Allow Reconcile to Mark for Edit")]
+		[Description("Switches Reconcile from preview mode to marking files for edit in the default CL")]
+		public bool AllowReconcileToMarkForEdit { get; set; }
+
+		[Category("Unreal.P4")]
+		[DisplayName("Reconcile Extensions")]
+		[Description("Reconcile will attempt to resolve these extensions (including wildcards)")]
+		public string ReconcileExtensions { get; set; } = "c*;h*;ini";
+
 		protected override void OnApply(PageApplyEventArgs e)
 		{
 			base.OnApply(e);
