UnrealVS: Fix remaining thread safety issues highlighted by static analysis.

* P4Commands now invokes a callback on UI thread for any UI operations
* StdOut and StdErr are now both returned when executing child processes to avoid contention on member variable
* Added locks around manipulating child process list.
* Fixed some issues where paths containing spaces were not quoted correctly
* Fixed detection of working directory where P4CONFIG is set to something other than p4config.txt.

diff --git a/Engine/Source/Programs/UnrealVS/ReleaseNotes.txt b/Engine/Source/Programs/UnrealVS/ReleaseNotes.txt
--- a/Engine/Source/Programs/UnrealVS/ReleaseNotes.txt
+++ b/Engine/Source/Programs/UnrealVS/ReleaseNotes.txt
@@ -1,7 +1,11 @@
-1.71 (WIP)
+1.72 (WIP)
 * fixed exception when opening project missing a referenced configuration
 * added exceptions to guard against code UI code not being accessed on the UI thread
 
+1.71
+* add p4FastReconcile
+* add ability to halt p4 output auto focus
+
 1.70
 * add option to allow async checkout
 * resolved concurrent command issues
diff --git a/Engine/Source/Programs/UnrealVS/UnrealVS.Shared/CommandLineEditor.cs b/Engine/Source/Programs/UnrealVS/UnrealVS.Shared/CommandLineEditor.cs
--- a/Engine/Source/Programs/UnrealVS/UnrealVS.Shared/CommandLineEditor.cs
+++ b/Engine/Source/Programs/UnrealVS/UnrealVS.Shared/CommandLineEditor.cs
@@ -114,7 +114,11 @@ namespace UnrealVS
 		/// </summary>
 		private void OnUpdateCommandLineCombo(Object _Source, ElapsedEventArgs _Args)
 		{
-			ThreadHelper.Generic.BeginInvoke(UpdateCommandLineCombo);
+			ThreadHelper.JoinableTaskFactory.Run(async () =>
+			{
+				await ThreadHelper.JoinableTaskFactory.SwitchToMainThreadAsync();
+				UpdateCommandLineCombo();
+			});
 		}
 
 		/// <summary>
diff --git a/Engine/Source/Programs/UnrealVS/UnrealVS.Shared/CompileSingleFile.cs b/Engine/Source/Programs/UnrealVS/UnrealVS.Shared/CompileSingleFile.cs
--- a/Engine/Source/Programs/UnrealVS/UnrealVS.Shared/CompileSingleFile.cs
+++ b/Engine/Source/Programs/UnrealVS/UnrealVS.Shared/CompileSingleFile.cs
@@ -18,7 +18,6 @@ namespace UnrealVS
 		private const int CompileSingleFileButtonID = 0x1075;
 		private const int PreprocessSingleFileButtonID = 0x1076;
 		private const int UBTSubMenuID = 0x3103;
-		private string	  PPLine = "";
 		private string	  FileToCompileOriginalExt = "";
 
 		static readonly List<string> ValidExtensions = new List<string> { ".c", ".cc", ".cpp", ".cxx" };
@@ -182,6 +181,7 @@ namespace UnrealVS
 			DTE.Events.BuildEvents.OnBuildBegin += BuildEvents_OnBuildBegin;
 
 			// Create a delegate for handling output messages
+			string PPLine = String.Empty;
 			void OutputHandler(object Sender, DataReceivedEventArgs Args) 
 			{ 
 				if (Args.Data != null) 
@@ -200,7 +200,6 @@ namespace UnrealVS
 			BuildCommandLine = BuildCommandLine.Replace("$(SolutionDir)", SolutionDir);
 			BuildCommandLine = BuildCommandLine.Replace("$(ProjectName)", VCStartupProject.Name);
 
-			PPLine = "";
 			FileToCompileOriginalExt = FileToCompileExt;
 
 			string PreProcess = bPreProcessOnly ? " -NoXGE -Preprocess " : "";
@@ -221,7 +220,7 @@ namespace UnrealVS
 				// add an event handler to respond to the exit of the preprocess request
 				// and open the generated file if it exists.
 				ChildProcess.EnableRaisingEvents = true;
-				ChildProcess.Exited += new EventHandler(PreprocessExitHandler);
+				ChildProcess.Exited += new EventHandler((s, e) => PreprocessExitHandler(PPLine));
 			}
 			
 			ChildProcess.Start();
@@ -231,31 +230,42 @@ namespace UnrealVS
 			return true;
 		}
 
-		private void PreprocessExitHandler(object sender, System.EventArgs e)
+		private void PreprocessExitHandler(string PPLine)
 		{
 			// not all compile actions support pre-process - check it exists
 			if (PPLine.Contains("PreProcessPath:"))
 			{
 				string PPFullPath = PPLine.Replace("PreProcessPath:", "").Trim();
 
-				IVsOutputWindowPane BuildOutputPane = UnrealVSPackage.Instance.GetOutputPane();
-				DTE DTE = UnrealVSPackage.Instance.DTE;
+				ThreadHelper.JoinableTaskFactory.Run(async () =>
+				{
+					await ThreadHelper.JoinableTaskFactory.SwitchToMainThreadAsync();
+					OpenPreprocessedFile(PPFullPath);
+				});
+			}
+		}
 
-				BuildOutputPane.OutputString($"1>  PPFullPath: {PPFullPath}{Environment.NewLine}");
+		private void OpenPreprocessedFile(string PPFullPath)
+		{
+			ThreadHelper.ThrowIfNotOnUIThread();
 
-				if (File.Exists(PPFullPath))
-				{
-					// if the file exists, rename it to isolate the file and have its extension be the original to maintain syntax highlighting
-					string Dir = Path.GetDirectoryName(PPFullPath);
-					string FileName = Path.GetFileNameWithoutExtension(Path.GetFileNameWithoutExtension(PPFullPath)) + "_preprocessed";
+			IVsOutputWindowPane BuildOutputPane = UnrealVSPackage.Instance.GetOutputPane();
+			DTE DTE = UnrealVSPackage.Instance.DTE;
 
-					string RenamedFile = Path.Combine(Dir, FileName) + FileToCompileOriginalExt;
+			BuildOutputPane.OutputString($"1>  PPFullPath: {PPFullPath}{Environment.NewLine}");
 
-					File.Copy(PPFullPath, RenamedFile, true /*overwrite*/);
+			if (File.Exists(PPFullPath))
+			{
+				// if the file exists, rename it to isolate the file and have its extension be the original to maintain syntax highlighting
+				string Dir = Path.GetDirectoryName(PPFullPath);
+				string FileName = Path.GetFileNameWithoutExtension(Path.GetFileNameWithoutExtension(PPFullPath)) + "_preprocessed";
 
-					DTE.ExecuteCommand("File.OpenFile", $"\"{RenamedFile}\"");
-				}
-			}	
+				string RenamedFile = Path.Combine(Dir, FileName) + FileToCompileOriginalExt;
+
+				File.Copy(PPFullPath, RenamedFile, true /*overwrite*/);
+
+				DTE.ExecuteCommand("File.OpenFile", $"\"{RenamedFile}\"");
+			}
 		}
 
 		private void BuildEvents_OnBuildBegin(vsBuildScope Scope, vsBuildAction Action)
diff --git a/Engine/Source/Programs/UnrealVS/UnrealVS.Shared/P4Commands.cs b/Engine/Source/Programs/UnrealVS/UnrealVS.Shared/P4Commands.cs
--- a/Engine/Source/Programs/UnrealVS/UnrealVS.Shared/P4Commands.cs
+++ b/Engine/Source/Programs/UnrealVS/UnrealVS.Shared/P4Commands.cs
@@ -14,6 +14,7 @@ using System.Text;
 using System.Text.RegularExpressions;
 using System.Threading;
 using System.Threading.Tasks;
+using Process = System.Diagnostics.Process;
 
 namespace UnrealVS
 {
@@ -25,9 +26,6 @@ namespace UnrealVS
 		private const bool bOutputStdOutOn = true;
 		private const bool bOutputStdOutOff = false;
 
-		private const bool bWaitForCompletionOn = true;
-		private const bool bWaitForCompletionOff = false;
-
 		private const int P4SubMenuID = 0x3100;
 		private const int P4CheckoutButtonID = 0x1450;
 		private const int P4AnnotateButtonID = 0x1451;
@@ -49,10 +47,6 @@ namespace UnrealVS
 		private bool bPullWorkingDirectorFromP4 = true;
 		private static Mutex bPullWorkingDirectorFromP4Mutex = new Mutex();
 
-		// stdXX from last operation
-		//private string P4OperationStdOut;
-		private string P4OperationStdErr;
-
 		// Exe paths
 		private string P4Exe = "C:\\Program Files\\Perforce\\p4.exe";
 		private string P4VCCmd = "C:\\Program Files\\Perforce\\p4vc.exe";
@@ -174,6 +168,8 @@ namespace UnrealVS
 		// Called when solutions are loaded or unloaded
 		private void SolutionOpened()
 		{
+			ThreadHelper.ThrowIfNotOnUIThread();
+
 			// Update the menu visibility
 			UpdateMenuOptions();
 		}
@@ -249,7 +245,7 @@ namespace UnrealVS
 
 		private void KillChildProcess()
 		{
-			foreach(var OldProcess in ChildProcessList)
+			foreach (var OldProcess in ChildProcessList)
 			{
 				if (!OldProcess.HasExited)
 				{
@@ -299,12 +295,13 @@ namespace UnrealVS
 				return;
 			}
 
-			string CaptureP4Output = "";
+			string CaptureP4Output;
+			string CaptureP4StdErr;
 
 			// Call annotate itself
-			bool bResult = TryP4Command($"annotate -TcIqu \"{DTE.ActiveDocument.FullName}", out CaptureP4Output, bPullWorkingDirectoryOn, bOutputStdOutOff);
+			bool bResult = TryP4Command($"annotate -TcIqu \"{DTE.ActiveDocument.FullName}", out CaptureP4Output, out CaptureP4StdErr, bPullWorkingDirectoryOn, bOutputStdOutOff);
 
-			if (!bResult || P4OperationStdErr.Length > 0)
+			if (!bResult || CaptureP4StdErr.Length > 0)
 			{
 				P4OutputPane.OutputString($"1>------ P4Annotate call failed");
 				return;
@@ -463,7 +460,7 @@ namespace UnrealVS
 
 			// get the HAVE revision
 			// p4 fstat -T "haveRev" -Olp //UE5/Main/Engine/Source/Programs/UnrealVS/UnrealVS.Shared/P4Commands.cs
-			if (TryP4Command($"fstat -T \"haveRev,depotFile\" -Olp \"{DTE.ActiveDocument.FullName}\"", out CaptureP4Output))
+			if (TryP4Command($"fstat -T \"haveRev,depotFile\" -Olp \"{DTE.ActiveDocument.FullName}\"", out CaptureP4Output, out _))
 			{
 				// expect output of the form
 				//		"... haveRev 5"
@@ -486,7 +483,7 @@ namespace UnrealVS
 				string CaptureP4Output2 = "";
 
 				string VersionPath = $"{depotPath}#{HaveRev}";
-				if (TryP4Command($"-q print \"{VersionPath}\"",out CaptureP4Output2, bPullWorkingDirectoryOn, bOutputStdOutOff))
+				if (TryP4Command($"-q print \"{VersionPath}\"", out CaptureP4Output2, out _, bPullWorkingDirectoryOn, bOutputStdOutOff))
 				{
 					File.WriteAllText(TempFilePath, CaptureP4Output2);
 
@@ -509,7 +506,7 @@ namespace UnrealVS
 			string ChangesTempFilePath = Path.Combine(TempPath, ChangesTempFileName);
 
 			string CaptureP4Output = "";
-			if (TryP4Command($"changes -itls submitted -m 10 {DTE.ActiveDocument.FullName}", out CaptureP4Output, bOutputStdOutOff))
+			if (TryP4Command($"changes -itls submitted -m 10 \"{DTE.ActiveDocument.FullName}\"", out CaptureP4Output, out _, bOutputStdOutOff))
 			{
 				File.WriteAllText(ChangesTempFilePath, CaptureP4Output);
 
@@ -521,6 +518,8 @@ namespace UnrealVS
 
 		private void P4FastReconcileCodeFiles(object Sender, EventArgs Args)
 		{
+			ThreadHelper.ThrowIfNotOnUIThread();
+
 			DTE DTE = UnrealVSPackage.Instance.DTE;
 
 			string[] DefaultExtensions =
@@ -563,7 +562,7 @@ namespace UnrealVS
 			P4OutputPane.OutputString($"Running Async Reconcile on : {ReconcileDepotPaths}{Environment.NewLine}");
 			P4OutputPane.Activate();
 
-			TryP4CommandAsync($"reconcile -e -m {Preview} {ReconcileDepotPaths}");
+			_ = TryP4CommandAsync($"reconcile -e -m {Preview} {ReconcileDepotPaths}");
 		}
 
 		private void P4ShowFileInP4Vandler(object Sender, EventArgs Args)
@@ -577,16 +576,18 @@ namespace UnrealVS
 				return;
 			}
 
-			TryP4VCommand($"-s {DTE.ActiveDocument.FullName}");
+			TryP4VCommand($"-s \"{DTE.ActiveDocument.FullName}\"");
 		}
 		public void OpenForEdit(string FileName, bool CheckOptionsFlag = false)
 		{
+			ThreadHelper.ThrowIfNotOnUIThread();
+
 			// if the callee requested it, honor the enabling of autocheckout option - return if its off
 			if (CheckOptionsFlag && !UnrealVSPackage.Instance.OptionsPage.AllowUnrealVSCheckoutOnEdit)
 			{
 				//P4OutputPane.OutputString($"autocheckout disabled: {FileName}");
 				return;
-			}	
+			}
 
 			// Don't open for edit if the file is already writable
 			if (!File.Exists(FileName) || !File.GetAttributes(FileName).HasFlag(FileAttributes.ReadOnly))
@@ -603,18 +604,19 @@ namespace UnrealVS
 				File.SetAttributes(FileName, NewAttributes);
 
 				// then send a fire and forget p4 edit command - does not lock the UI
-				TryP4CommandAsync($"edit \"{FileName}");
+				_ = TryP4CommandAsync($"edit \"{FileName}\"");
 			}
 			else
 			{
-				string CaptureP4Output = "";
 				// sync edit command - will lock the UI but be safer (no risk of a locally writeable file that is not checked out)
-				TryP4Command($"edit \"{FileName}", out CaptureP4Output);
+				TryP4Command($"edit \"{FileName}\"", out _, out _);
 			}
 		}
 
 		private string ReadWorkingDirectorFromP4Config()
 		{
+			ThreadHelper.ThrowIfNotOnUIThread();
+
 			string WorkingDirectory = "";
 			//------P4 Operation started
 			//P4CLIENT = andrew.firth_private_frosty2(config 'd:\p4\frosty\p4config.txt')
@@ -625,30 +627,25 @@ namespace UnrealVS
 			//P4_perforce:1666_CHARSET = none(set)
 
 			string CaptureP4Output = "";
-			TryP4Command("set", out CaptureP4Output, bPullWorkingDirectoryOff);
+			TryP4Command("set", out CaptureP4Output, out _, bPullWorkingDirectoryOff);
 
 			bool bSuccess = false;
 
-			if (CaptureP4Output.Length > 0)
+			string[] lines = CaptureP4Output.Split('\n');
+			foreach(string Line in lines)
 			{
-				string[] lines = CaptureP4Output.Split('\n');
-
-				// very brittle - Index assume there is a better way leveraging C#
-				if (lines.Length > 0)
+				System.Text.RegularExpressions.Match Match = Regex.Match(Line, @"^P4CONFIG\s*=.*'([^']+)'\)$");
+				if (Match.Success)
 				{
-					if (lines[0].Contains("p4config.txt"))
-					{
-						string path = lines[0].Split('\'')[1];
-
-						WorkingDirectory = Path.GetDirectoryName(path);
-						bSuccess = true;
-					}
+					WorkingDirectory = Path.GetDirectoryName(Match.Groups[1].Value);
+					bSuccess = true;
+					break;
 				}
 			}
 
 			if (!bSuccess)
 			{
-				P4OutputPane.OutputString($"attempt to pull p4config.txt info failed{Environment.NewLine}");
+				P4OutputPane.OutputString($"attempt to pull P4CONFIG info failed{Environment.NewLine}");
 			}
 
 			return WorkingDirectory;
@@ -656,8 +653,10 @@ namespace UnrealVS
 
 		void SetUserInfoStrings()
 		{
-			string CaptureP4Output = "";
-			TryP4Command($"-s -L \"{P4WorkingDirectory}\" info", out CaptureP4Output, bPullWorkingDirectoryOff);
+			ThreadHelper.ThrowIfNotOnUIThread();
+
+			string CaptureP4Output;
+			TryP4Command($"-s -L \"{P4WorkingDirectory}\" info", out CaptureP4Output, out _, bPullWorkingDirectoryOff);
 
 			Regex UserPattern = new Regex(@"User name: (?<user>.*)$", RegexOptions.Compiled | RegexOptions.Multiline);
 			Regex PortPattern = new Regex(@"Server address: (?<port>.*)$", RegexOptions.Compiled | RegexOptions.Multiline);
@@ -684,81 +683,96 @@ namespace UnrealVS
 			ThreadHelper.ThrowIfNotOnUIThread();
 
 			bPullWorkingDirectorFromP4Mutex.WaitOne();
-
-			if (IsSolutionLoaded() && (P4WorkingDirectory == null || P4WorkingDirectory.Length < 2))
+			try
 			{
-				DTE DTE = UnrealVSPackage.Instance.DTE;
-
-				// use the current solution folder as a temp working directory
-				P4WorkingDirectory = Path.GetDirectoryName(DTE.Solution.FileName);
+				if (IsSolutionLoaded() && (P4WorkingDirectory == null || P4WorkingDirectory.Length < 2))
+				{
+					DTE DTE = UnrealVSPackage.Instance.DTE;
 
-				// SetUserInfoStrings 
-				SetUserInfoStrings();
-			}
+					// use the current solution folder as a temp working directory
+					P4WorkingDirectory = Path.GetDirectoryName(DTE.Solution.FileName);
 
-			// if the callee wants us to pull a CWD and it wasn't already done
-			// pull it from the p4 config now
-			if (bPullfromP4Settings && bPullWorkingDirectorFromP4)
-			{
-				string NewWorkingDirectory = ReadWorkingDirectorFromP4Config();
+					// SetUserInfoStrings 
+					SetUserInfoStrings();
+				}
 
-				if (NewWorkingDirectory.Length > 1)
+				// if the callee wants us to pull a CWD and it wasn't already done
+				// pull it from the p4 config now
+				if (bPullfromP4Settings && bPullWorkingDirectorFromP4)
 				{
-					P4WorkingDirectory = NewWorkingDirectory;
-					bPullWorkingDirectorFromP4 = false;
+					string NewWorkingDirectory = ReadWorkingDirectorFromP4Config();
 
-					P4OutputPane.OutputString($"P4WorkingDirectory set to '{P4WorkingDirectory}'{Environment.NewLine}");
-				}
-				else
-				{
-					P4OutputPane.OutputString($"P4WorkingDirectory set failed {Environment.NewLine}");
+					if (NewWorkingDirectory.Length > 1)
+					{
+						P4WorkingDirectory = NewWorkingDirectory;
+						bPullWorkingDirectorFromP4 = false;
+
+						P4OutputPane.OutputString($"P4WorkingDirectory set to '{P4WorkingDirectory}'{Environment.NewLine}");
+					}
+					else
+					{
+						P4OutputPane.OutputString($"P4WorkingDirectory set failed {Environment.NewLine}");
+					}
 				}
 			}
-
-			bPullWorkingDirectorFromP4Mutex.ReleaseMutex();
+			finally
+			{
+				bPullWorkingDirectorFromP4Mutex.ReleaseMutex();
+			}
 		}
 
 		public async Task<bool> TryP4CommandAsync(string CommandLine, bool bPullWorkingDirectoryNow = bPullWorkingDirectoryOn, bool bOutputStdOut = bOutputStdOutOn)
 		{
-			string CaptureP4Output;
-			return await System.Threading.Tasks.Task.Run( () => TryP4Command(CommandLine, out CaptureP4Output, bPullWorkingDirectoryNow, bOutputStdOut) );
+			(bool Result, _, _) = await TryP4CommandExAsync(P4Exe, CommandLine, bPullWorkingDirectoryNow, bOutputStdOut);
+			return Result;
 		}
 
-		private bool TryP4Command(string CommandLine, out string CaptureP4Output, bool bPullWorkingDirectoryNow = bPullWorkingDirectoryOn, bool bOutputStdOut = bOutputStdOutOn)
+		private bool TryP4Command(string CommandLine, out string CaptureP4Output, out string CaptureP4StdErr, bool bPullWorkingDirectoryNow = bPullWorkingDirectoryOn, bool bOutputStdOut = bOutputStdOutOn)
 		{
-			return TryP4CommandEx(P4Exe, CommandLine, out CaptureP4Output, bPullWorkingDirectoryNow, bOutputStdOut);
+			return TryP4CommandEx(P4Exe, CommandLine, out CaptureP4Output, out CaptureP4StdErr, bPullWorkingDirectoryNow, bOutputStdOut);
 		}
 
 		private bool TryP4VCCommand(string CommandLine, bool bPullWorkingDirectoryNow = bPullWorkingDirectoryOn, bool bOutputStdOut = bOutputStdOutOn)
 		{
 			if (P4VCCmd.Length > 1)
 			{
-				string ignore;
-				return TryP4CommandEx(P4VCCmd, CommandLine, out ignore, bPullWorkingDirectoryNow, bOutputStdOut);
+				return TryP4CommandEx(P4VCCmd, CommandLine, out _, out _, bPullWorkingDirectoryNow, bOutputStdOut);
 			}
 
 			return false;
 		}
 
-		private bool TryP4VCommand(string CommandLine, bool bPullWorkingDirectoryNow = bPullWorkingDirectoryOn, bool bOutputStdOut = bOutputStdOutOn)
+		private void TryP4VCommand(string CommandLine, bool bPullWorkingDirectoryNow = bPullWorkingDirectoryOn, bool bOutputStdOut = bOutputStdOutOn)
 		{
-			string ignore;
-			return TryP4CommandEx(P4VExe, CommandLine, out ignore, bPullWorkingDirectoryNow, bOutputStdOut, bWaitForCompletionOff);
+			_ = TryP4CommandExAsync(P4VExe, CommandLine, bPullWorkingDirectoryNow, bOutputStdOut);
 		}
 
-		private bool TryP4CommandEx(string CmdPath, string CommandLine, out string CaptureP4Output, bool bPullWorkingDirectoryNow = bPullWorkingDirectoryOn, bool bOutputStdOut = bOutputStdOutOn, bool bWaitForCompletion = bWaitForCompletionOn)
+		private bool TryP4CommandEx(string CmdPath, string CommandLine, out string CaptureP4StdOut, out string CaptureP4StdErr, bool bPullWorkingDirectoryNow = bPullWorkingDirectoryOn, bool bOutputStdOut = bOutputStdOutOn)
 		{
-			PullWorkingDirectory(bPullWorkingDirectoryNow);
+			(bool Result, string StdOut, string StdErr) = ThreadHelper.JoinableTaskFactory.Run(() => TryP4CommandExAsync(CmdPath, CommandLine, bPullWorkingDirectoryNow, bOutputStdOut));
 
-			DTE DTE = UnrealVSPackage.Instance.DTE;
+			CaptureP4StdOut = StdOut;
+			CaptureP4StdErr = StdErr;
 
-			// Set up the output pane
-			if (UnrealVSPackage.Instance.OptionsPage.ForceOutputWindow)
+			return Result;
+		}
+
+		private async Task<(bool Result, string StdOut, string StdErr)> TryP4CommandExAsync(string CmdPath, string CommandLine, bool bPullWorkingDirectoryNow = bPullWorkingDirectoryOn, bool bOutputStdOut = bOutputStdOutOn)
+		{
+			await ThreadHelper.JoinableTaskFactory.RunAsync(async () =>
 			{
-				P4OutputPane.Activate();
-			}
-			
-			
+				await ThreadHelper.JoinableTaskFactory.SwitchToMainThreadAsync();
+				PullWorkingDirectory(bPullWorkingDirectoryNow);
+
+				// Set up the output pane
+				if (UnrealVSPackage.Instance.OptionsPage.ForceOutputWindow)
+				{
+					P4OutputPane.Activate();
+				}
+
+				P4OutputPane.OutputString($"Running \"{CmdPath}\" {CommandLine}{Environment.NewLine}");
+			});
+
 			StringBuilder StdOutSB = new StringBuilder();
 			StringBuilder StdErrSB = new StringBuilder();
 			DateTime Start = DateTime.Now;
@@ -780,66 +794,52 @@ namespace UnrealVS
 			// Create a delegate for handling output messages
 			ChildProcess.OutputDataReceived += (s, a) => { if (a.Data != null) StdOutSB.AppendLine(a.Data); };
 			ChildProcess.ErrorDataReceived += (s, a) => { if (a.Data != null) StdErrSB.AppendLine(a.Data); };
-			if (bWaitForCompletion == false)
-			{
 				ChildProcess.EnableRaisingEvents = true;
-				ChildProcess.Exited += (s, a) =>
-				{
-					string P4Output = StdOutSB.ToString();
-					P4OperationStdErr = StdErrSB.ToString();
-
-					if (P4Output.Length > 0 && bOutputStdOut)
-					{
-						P4OutputPane.OutputString(P4Output);
-					}
-
-					if (P4OperationStdErr.Length > 0)
-					{
-						P4OutputPane.OutputString(P4OperationStdErr);
-					}
 
-					TimeSpan Duration = DateTime.Now - Start;
-					P4OutputPane.OutputString($"complete {Duration.TotalSeconds.ToString()} {Environment.NewLine}");
+			TaskCompletionSource<bool> ProcessExitTaskSource = new TaskCompletionSource<bool>();
+			ChildProcess.Exited += (s, a) =>
+			{
+				ProcessExitTaskSource.TrySetResult(true);
+			};
 
-					ChildProcessList.Remove(ChildProcess);
-				};
+			lock (ChildProcessList)
+			{
+				ChildProcessList.Add(ChildProcess);
 			}
 
-			ChildProcessList.Add(ChildProcess);
-
 			ChildProcess.Start();
 			ChildProcess.BeginOutputReadLine();
 			ChildProcess.BeginErrorReadLine();
 
-			if (bWaitForCompletion)
-			{ 
-				ChildProcess.WaitForExit();
-				TimeSpan Duration = DateTime.Now - Start;
+			await ProcessExitTaskSource.Task;
+			ChildProcess.WaitForExit();
 
-				CaptureP4Output = StdOutSB.ToString();
-				P4OperationStdErr = StdErrSB.ToString();
+			lock (ChildProcessList)
+			{
+				ChildProcessList.Remove(ChildProcess);
+			}
 
-				if (CaptureP4Output.Length > 0 && bOutputStdOut)
+			string StdOut = StdOutSB.ToString();
+			string StdErr = StdErrSB.ToString();
+
+			await ThreadHelper.JoinableTaskFactory.RunAsync(async () =>
+			{
+				await ThreadHelper.JoinableTaskFactory.SwitchToMainThreadAsync();
+				if (StdOut.Length > 0 && bOutputStdOut)
 				{
-					P4OutputPane.OutputString(CaptureP4Output);
+					P4OutputPane.OutputString(StdOut);
 				}
-
-				if (P4OperationStdErr.Length > 0)
+				if (StdErr.Length > 0)
 				{
-					P4OutputPane.OutputString(P4OperationStdErr);
+					P4OutputPane.OutputString(StdErr);
 				}
 
-				P4OutputPane.OutputString($"complete {((int)Duration.TotalSeconds).ToString()}seconds {Environment.NewLine}");
-
-				ChildProcessList.Remove(ChildProcess);
+				TimeSpan Duration = DateTime.Now - Start;
+				P4OutputPane.OutputString($"complete {Duration.TotalSeconds.ToString()} {Environment.NewLine}");
+			});
 
-				return P4OperationStdErr.Length == 0;
-			}
-			else
-			{
-				CaptureP4Output = "async operations do not capture p4 output";
-				return true;
-			}
+			bool Result = (ChildProcess.ExitCode == 0 && StdErr.Length == 0);
+			return (Result, StdOut, StdErr);
 		}
 	}
 
@@ -858,6 +858,8 @@ namespace UnrealVS
 
 		public int OnBeforeSave(uint DocumentCookie)
 		{
+			ThreadHelper.ThrowIfNotOnUIThread();
+
 			RunningDocumentInfo DocumentInfo = RunningDocumentTable.GetDocumentInfo(DocumentCookie);
 
 			string AbsoluteFilePath = DocumentInfo.Moniker;
diff --git a/Engine/Source/Programs/UnrealVS/UnrealVS.Shared/ReleaseNotes.txt b/Engine/Source/Programs/UnrealVS/UnrealVS.Shared/ReleaseNotes.txt
deleted file mode 100644
diff --git a/Engine/Source/Programs/UnrealVS/UnrealVS.Shared/UnrealVSPackage.cs b/Engine/Source/Programs/UnrealVS/UnrealVS.Shared/UnrealVSPackage.cs
--- a/Engine/Source/Programs/UnrealVS/UnrealVS.Shared/UnrealVSPackage.cs
+++ b/Engine/Source/Programs/UnrealVS/UnrealVS.Shared/UnrealVSPackage.cs
@@ -24,6 +24,7 @@ using Microsoft.VisualStudio.Text.Editor;
 using IOleServiceProvider = Microsoft.VisualStudio.OLE.Interop.IServiceProvider;
 using Microsoft.VisualStudio.Utilities;
 using Microsoft.VisualStudio.Text.Projection;
+using Microsoft;
 
 namespace UnrealVS
 {
@@ -192,7 +193,6 @@ namespace UnrealVS
 			Logging.WriteLine("Loading UnrealVS extension package...");
 		}
 
-
 		/// <summary>
 		/// Initializes the package right after it's been "sited" into the fully-initialized Visual Studio IDE.
 		/// </summary>
@@ -208,6 +208,7 @@ namespace UnrealVS
 			// Get access to Visual Studio's DTE object.  This object has various hooks into the Visual Studio
 			// shell that are useful for writing extensions.
 			DTE = await GetServiceAsync(typeof(DTE)) as DTE;
+			Assumes.Present(DTE);
 			Logging.WriteLine("DTE version " + DTE.Version);
 
 			//TextManager = await GetServiceAsync(typeof(VsTextManagerClass)) as IVsTextManager3;
@@ -218,18 +219,20 @@ namespace UnrealVS
 			// Get selection manager and register to receive events
 			SelectionManager =
 				await GetServiceAsync(typeof(SVsShellMonitorSelection)) as IVsMonitorSelection;
+			Assumes.Present(SelectionManager);
 			SelectionManager.AdviseSelectionEvents(this, out SelectionEventsHandle);
 
 			// Get solution and register to receive events
 			SolutionManager = await GetServiceAsync(typeof(SVsSolution)) as IVsSolution2;
+			Assumes.Present(SolutionManager);
 			UpdateUnrealLoadedStatus();
 			SolutionManager.AdviseSolutionEvents(this, out SolutionEventsHandle);
 
 			// Grab the solution build manager.  We need this in order to change certain things about the Visual
 			// Studio environment, like what the active startup project is
 			// Get solution build manager
-			SolutionBuildManager =
-				await GetServiceAsync(typeof(SVsSolutionBuildManager)) as IVsSolutionBuildManager2;
+			SolutionBuildManager = await GetServiceAsync(typeof(SVsSolutionBuildManager)) as IVsSolutionBuildManager2;
+			Assumes.Present(SolutionBuildManager);
 			SolutionBuildManager.AdviseUpdateSolutionEvents(this, out UpdateSolutionEventsHandle);
 
 			// Create our command-line editor
@@ -295,7 +298,12 @@ namespace UnrealVS
 				while (true)
 				{
 					if (bCancelTicker != 0) return;
-					ThreadHelper.Generic.BeginInvoke(Tick);
+
+					ThreadHelper.JoinableTaskFactory.Run(async () =>
+					{
+						await ThreadHelper.JoinableTaskFactory.SwitchToMainThreadAsync();
+						Tick();
+					});
 
 					if (bCancelTicker != 0) return;
 					Thread.Sleep(TickPeriod);
@@ -351,6 +359,9 @@ namespace UnrealVS
 			CompileSingleFile?.Dispose();
 			CompileSingleFile = null;
 
+			P4CommandsGroup?.Dispose();
+			P4CommandsGroup = null;
+
 			CommandLineEditor?.Dispose();
 			CommandLineEditor = null;
 
