[UnrealVS] - enable warnings as errors

diff --git a/Engine/Source/Programs/UnrealVS/UnrealVS.2019/UnrealVS.2019.csproj b/Engine/Source/Programs/UnrealVS/UnrealVS.2019/UnrealVS.2019.csproj
--- a/Engine/Source/Programs/UnrealVS/UnrealVS.2019/UnrealVS.2019.csproj
+++ b/Engine/Source/Programs/UnrealVS/UnrealVS.2019/UnrealVS.2019.csproj
@@ -180,6 +180,7 @@
     <ErrorReport>prompt</ErrorReport>
     <CodeAnalysisRuleSet>MinimumRecommendedRules.ruleset</CodeAnalysisRuleSet>
     <NoWarn>MSB3277;NU1603</NoWarn>
+    <TreatWarningsAsErrors>true</TreatWarningsAsErrors>
   </PropertyGroup>
   <PropertyGroup Condition=" '$(Configuration)|$(Platform)' == 'Release|AnyCPU' ">
     <OutputPath>bin\Release\</OutputPath>
diff --git a/Engine/Source/Programs/UnrealVS/UnrealVS.2022/UnrealVS.2022.csproj b/Engine/Source/Programs/UnrealVS/UnrealVS.2022/UnrealVS.2022.csproj
--- a/Engine/Source/Programs/UnrealVS/UnrealVS.2022/UnrealVS.2022.csproj
+++ b/Engine/Source/Programs/UnrealVS/UnrealVS.2022/UnrealVS.2022.csproj
@@ -97,7 +97,7 @@
     <Resource Include="Resources\BatchBuilderBackground.png" />
     <Resource Include="Resources\Icon.ico" />
     <Content Include="..\ReleaseNotes.txt">
-	  <Visible>false</Visible>
+      <Visible>false</Visible>
       <CopyToOutputDirectory>Always</CopyToOutputDirectory>
       <IncludeInVSIX>true</IncludeInVSIX>
     </Content>
@@ -190,6 +190,7 @@
     <CodeAnalysisRuleSet>MinimumRecommendedRules.ruleset</CodeAnalysisRuleSet>
     <CopyVsixExtensionFiles>False</CopyVsixExtensionFiles>
     <NoWarn>MSB3277</NoWarn>
+    <TreatWarningsAsErrors>true</TreatWarningsAsErrors>
   </PropertyGroup>
   <Target Name="IncludeNuGetPackageReferences" AfterTargets="GetVsixSourceItems">
     <ItemGroup>
