[UnrealVS] - convert reconcile over to using local paths, add up* to the paths + small code review feedback

(no new version of the VSIX)

diff --git a/Engine/Source/Programs/UnrealVS/UnrealVS.Shared/P4Commands.cs b/Engine/Source/Programs/UnrealVS/UnrealVS.Shared/P4Commands.cs
--- a/Engine/Source/Programs/UnrealVS/UnrealVS.Shared/P4Commands.cs
+++ b/Engine/Source/Programs/UnrealVS/UnrealVS.Shared/P4Commands.cs
@@ -526,7 +526,9 @@ namespace UnrealVS
 			{
 				"c*",
 				"h*",
-				"ini"
+				"ini",
+				"uproject",
+				"uplugin"
 			};
 
 			string[] Extensions = UnrealVSPackage.Instance.OptionsPage.ReconcileExtensions.Split(';');
@@ -540,14 +542,14 @@ namespace UnrealVS
 			string ReconcileDepotPaths = "";
 			string Preview = "-n";
 
-			if (Stream.Length < 1)
+			if (string.IsNullOrWhiteSpace(P4WorkingDirectory))
 			{
 				PullWorkingDirectory(bPullWorkingDirectoryOn);
 			}
 
 			foreach (string Ext in Extensions)
 			{
-				ReconcileDepotPaths += Stream + "/...." + Ext + " ";
+				ReconcileDepotPaths += P4WorkingDirectory + Path.DirectorySeparatorChar + "...." + Ext + " ";
 			}
 
 			if (UnrealVSPackage.Instance.OptionsPage.AllowReconcileToMarkForEdit)
diff --git a/Engine/Source/Programs/UnrealVS/UnrealVS.Shared/UnrealVSOptions.cs b/Engine/Source/Programs/UnrealVS/UnrealVS.Shared/UnrealVSOptions.cs
--- a/Engine/Source/Programs/UnrealVS/UnrealVS.Shared/UnrealVSOptions.cs
+++ b/Engine/Source/Programs/UnrealVS/UnrealVS.Shared/UnrealVSOptions.cs
@@ -50,7 +50,7 @@ namespace UnrealVS
 		[Category("Unreal.P4")]
 		[DisplayName("Reconcile Extensions")]
 		[Description("Reconcile will attempt to resolve these extensions (including wildcards)")]
-		public string ReconcileExtensions { get; set; } = "c*;h*;ini";
+		public string ReconcileExtensions { get; set; } = "c*;h*;ini;uproject;uplugin";
 
 		protected override void OnApply(PageApplyEventArgs e)
 		{
