Added commands to open File History, Revision Graph and Timelapse views from the UnrealVS.P4 menu.

Don't abort p4 checkout if the file is writable if the command is executed manually.

diff --git a/Engine/Source/Programs/UnrealVS/UnrealVS.Shared/P4Commands.cs b/Engine/Source/Programs/UnrealVS/UnrealVS.Shared/P4Commands.cs
--- a/Engine/Source/Programs/UnrealVS/UnrealVS.Shared/P4Commands.cs
+++ b/Engine/Source/Programs/UnrealVS/UnrealVS.Shared/P4Commands.cs
@@ -34,7 +34,10 @@ namespace UnrealVS
 		private const int P4DiffinVSButtonID = 0x1454;
 		private const int P4GetLast10ChangesID = 0x1455;
 		private const int P4ShowFileinP4VID = 0x1456;
-		private const int P4FastReconcileCodeFilesID  =0x1457;
+		private const int P4FastReconcileCodeFilesID = 0x1457;
+		private const int P4TimelapseButtonID = 0x1458;
+		private const int P4RevisionGraphButtonID = 0x1459;
+		private const int P4FileHistoryButtonID = 0x1460;
 
 		private OleMenuCommand SubMenuCommand;
 
@@ -122,14 +125,14 @@ namespace UnrealVS
 			// figure out the P4VC path
 			if (!File.Exists(P4VCCmd))
 			{
-				if (File.Exists(P4VCCmd))
+				if (File.Exists(P4VCCmdBat))
 				{
 					P4VCCmd = P4VCCmdBat;
 				}
 				else
 				{
 					P4OutputPane.Activate();
-					P4OutputPane.OutputString($"1>------ P4VC not found, {P4VCCmd} or {P4VCCmd}{Environment.NewLine}");
+					P4OutputPane.OutputString($"1>------ P4VC not found, {P4VCCmd} or {P4VCCmdBat}{Environment.NewLine}");
 					P4VCCmd = "";
 				}
 
@@ -138,17 +141,18 @@ namespace UnrealVS
 			// add commands
 			P4CommandsList.Add(new P4Command(P4CheckoutButtonID, P4CheckoutButtonHandler));
 			P4CommandsList.Add(new P4Command(P4AnnotateButtonID, P4AnnotateButtonHandler));
+			P4CommandsList.Add(new P4Command(P4TimelapseButtonID, P4TimelapseHandler));
 			P4CommandsList.Add(new P4Command(P4IntegrationAwareTimelapseButtonID, P4IntegrationAwareTimeLapseHandler));
+			P4CommandsList.Add(new P4Command(P4RevisionGraphButtonID, P4RevisionGraphHandler));
 			P4CommandsList.Add(new P4Command(P4DiffinVSButtonID, P4DiffinVSHandler));
 			P4CommandsList.Add(new P4Command(P4GetLast10ChangesID, P4GetLast10ChangesHandler));
-			P4CommandsList.Add(new P4Command(P4ShowFileinP4VID, P4ShowFileInP4Vandler));
+			P4CommandsList.Add(new P4Command(P4FileHistoryButtonID, P4FileHistoryHandler));
+			P4CommandsList.Add(new P4Command(P4ShowFileinP4VID, P4ShowFileInP4VHandler));
 			P4CommandsList.Add(new P4Command(P4FastReconcileCodeFilesID, P4FastReconcileCodeFiles));
-			
 
 			if (P4VCCmd.Length > 1)
 			{
 				P4CommandsList.Add(new P4Command(P4ViewSelectedCLButtonID, P4ViewSelectedCLButtonHandler));
-
 			}
 
 			// add sub menu for commands
@@ -407,6 +411,57 @@ namespace UnrealVS
 
 		}
 
+		private void P4TimelapseHandler(object sender, EventArgs args)
+		{
+			ThreadHelper.ThrowIfNotOnUIThread();
+
+			DTE DTE = UnrealVSPackage.Instance.DTE;
+
+			if (DTE.ActiveDocument == null)
+			{
+				Logging.WriteLine("P4Timelapse called without an active document");
+				P4OutputPane.OutputString($"1>------ P4Timelapse called without an active document{Environment.NewLine}");
+				return;
+			}
+
+			string Command = $"timelapse \"{DTE.ActiveDocument.FullName}\"";
+			TryP4VCCommand(Command, true);
+		}
+
+		private void P4RevisionGraphHandler(object sender, EventArgs args)
+		{
+			ThreadHelper.ThrowIfNotOnUIThread();
+
+			DTE DTE = UnrealVSPackage.Instance.DTE;
+
+			if (DTE.ActiveDocument == null)
+			{
+				Logging.WriteLine("P4Timelapse called without an active document");
+				P4OutputPane.OutputString($"1>------ P4RevisionGraph called without an active document{Environment.NewLine}");
+				return;
+			}
+
+			string Command = $"revgraph \"{DTE.ActiveDocument.FullName}\"";
+			TryP4VCCommand(Command, true);
+		}
+
+		private void P4FileHistoryHandler(object sender, EventArgs args)
+		{
+			ThreadHelper.ThrowIfNotOnUIThread();
+
+			DTE DTE = UnrealVSPackage.Instance.DTE;
+
+			if (DTE.ActiveDocument == null)
+			{
+				Logging.WriteLine("P4FileHistory called without an active document");
+				P4OutputPane.OutputString($"1>------ P4FileHistory called without an active document{Environment.NewLine}");
+				return;
+			}
+
+			string Command = $"history \"{DTE.ActiveDocument.FullName}\"";
+			TryP4VCCommand(Command, true);
+		}
+
 		private void P4IntegrationAwareTimeLapseHandler(object Sender, EventArgs Args)
 		{
 			ThreadHelper.ThrowIfNotOnUIThread();
@@ -422,7 +477,7 @@ namespace UnrealVS
 				return;
 			}
 
-			string Command = $"-win 0 {UserInfoComplete} -cmd \"annotate -i \"{ DTE.ActiveDocument.FullName}\"\"";
+			string Command = $"-win 0 {UserInfoComplete} -cmd \"annotate -i \"{DTE.ActiveDocument.FullName}\"\"";
 
 			TryP4VCommand(Command);
 		}
@@ -567,7 +622,7 @@ namespace UnrealVS
 			_ = TryP4CommandAsync($"reconcile -e -m {Preview} {ReconcileDepotPaths}");
 		}
 
-		private void P4ShowFileInP4Vandler(object Sender, EventArgs Args)
+		private void P4ShowFileInP4VHandler(object Sender, EventArgs Args)
 		{
 			ThreadHelper.ThrowIfNotOnUIThread();
 
@@ -591,10 +646,11 @@ namespace UnrealVS
 				return;
 			}
 
-			// Don't open for edit if the file is already writable
-			if (!File.Exists(FileName) || !File.GetAttributes(FileName).HasFlag(FileAttributes.ReadOnly))
+			// Don't open for edit if the file is already writable, unless the command was user-triggered
+			if (!File.Exists(FileName) ||
+ 				(CheckOptionsFlag && !File.GetAttributes(FileName).HasFlag(FileAttributes.ReadOnly)))
 			{
-				//P4OutputPane.OutputString($"already writeable: {FileName}");
+				P4OutputPane.OutputString($"already writeable: {FileName}");
 				return;
 			}
 
