UnrealVS: Fix writing to output pane on incorrect thread

#rnx

diff --git a/Engine/Source/Programs/UnrealVS/UnrealVS.2019/Properties/AssemblyInfo.cs b/Engine/Source/Programs/UnrealVS/UnrealVS.2019/Properties/AssemblyInfo.cs
--- a/Engine/Source/Programs/UnrealVS/UnrealVS.2019/Properties/AssemblyInfo.cs
+++ b/Engine/Source/Programs/UnrealVS/UnrealVS.2019/Properties/AssemblyInfo.cs
@@ -31,8 +31,8 @@ using System.Runtime.InteropServices;
 // You can specify all the values or you can default the Revision and Build Numbers 
 // by using the '*' as shown below:
 
-[assembly: AssemblyVersion("1.72.0.0")]
-[assembly: AssemblyFileVersion("1.72.0.0")]
+[assembly: AssemblyVersion("1.74.0.0")]
+[assembly: AssemblyFileVersion("1.74.0.0")]
 
 
 
diff --git a/Engine/Source/Programs/UnrealVS/UnrealVS.2022/Properties/AssemblyInfo.cs b/Engine/Source/Programs/UnrealVS/UnrealVS.2022/Properties/AssemblyInfo.cs
--- a/Engine/Source/Programs/UnrealVS/UnrealVS.2022/Properties/AssemblyInfo.cs
+++ b/Engine/Source/Programs/UnrealVS/UnrealVS.2022/Properties/AssemblyInfo.cs
@@ -30,8 +30,8 @@ using System.Runtime.InteropServices;
 // You can specify all the values or you can default the Revision and Build Numbers 
 // by using the '*' as shown below:
 
-[assembly: AssemblyVersion("1.72.0.0")]
-[assembly: AssemblyFileVersion("1.72.0.0")]
+[assembly: AssemblyVersion("1.74.0.0")]
+[assembly: AssemblyFileVersion("1.74.0.0")]
 
 
 
diff --git a/Engine/Source/Programs/UnrealVS/UnrealVS.Shared/GenerateProjectFiles.cs b/Engine/Source/Programs/UnrealVS/UnrealVS.Shared/GenerateProjectFiles.cs
--- a/Engine/Source/Programs/UnrealVS/UnrealVS.Shared/GenerateProjectFiles.cs
+++ b/Engine/Source/Programs/UnrealVS/UnrealVS.Shared/GenerateProjectFiles.cs
@@ -131,13 +131,17 @@ namespace UnrealVS
 		/// </summary>
 		public void OnOutputFromGenerateProjectFilesProcess(object _Sender, DataReceivedEventArgs Args)
 		{
-			ThreadHelper.ThrowIfNotOnUIThread();
-
-			var Pane = UnrealVSPackage.Instance.GetOutputPane();
-			if (Pane != null)
+			var AsyncTask = ThreadHelper.JoinableTaskFactory.RunAsync(async () =>
 			{
-				Pane.OutputString(Args.Data + "\n");
-			}
+				await ThreadHelper.JoinableTaskFactory.SwitchToMainThreadAsync();
+
+				var Pane = UnrealVSPackage.Instance.GetOutputPane();
+				if (Pane != null)
+				{
+					Pane.OutputString(Args.Data + "\n");
+				}
+			});
+			AsyncTask.Join();
 		}
 
 		private string GetBatchFileName()
