UnrealVS: Don't allow regenerating projects if the batch file isn't found

#rnx

diff --git a/Engine/Source/Programs/UnrealVS/UnrealVS.2019/Properties/AssemblyInfo.cs b/Engine/Source/Programs/UnrealVS/UnrealVS.2019/Properties/AssemblyInfo.cs
--- a/Engine/Source/Programs/UnrealVS/UnrealVS.2019/Properties/AssemblyInfo.cs
+++ b/Engine/Source/Programs/UnrealVS/UnrealVS.2019/Properties/AssemblyInfo.cs
@@ -31,8 +31,8 @@ using System.Runtime.InteropServices;
 // You can specify all the values or you can default the Revision and Build Numbers 
 // by using the '*' as shown below:
 
-[assembly: AssemblyVersion("1.74.0.0")]
-[assembly: AssemblyFileVersion("1.74.0.0")]
+[assembly: AssemblyVersion("1.75.0.0")]
+[assembly: AssemblyFileVersion("1.75.0.0")]
 
 
 
diff --git a/Engine/Source/Programs/UnrealVS/UnrealVS.2022/Properties/AssemblyInfo.cs b/Engine/Source/Programs/UnrealVS/UnrealVS.2022/Properties/AssemblyInfo.cs
--- a/Engine/Source/Programs/UnrealVS/UnrealVS.2022/Properties/AssemblyInfo.cs
+++ b/Engine/Source/Programs/UnrealVS/UnrealVS.2022/Properties/AssemblyInfo.cs
@@ -30,8 +30,8 @@ using System.Runtime.InteropServices;
 // You can specify all the values or you can default the Revision and Build Numbers 
 // by using the '*' as shown below:
 
-[assembly: AssemblyVersion("1.74.0.0")]
-[assembly: AssemblyFileVersion("1.74.0.0")]
+[assembly: AssemblyVersion("1.75.0.0")]
+[assembly: AssemblyFileVersion("1.75.0.0")]
 
 
 
diff --git a/Engine/Source/Programs/UnrealVS/UnrealVS.Shared/GenerateProjectFiles.cs b/Engine/Source/Programs/UnrealVS/UnrealVS.Shared/GenerateProjectFiles.cs
--- a/Engine/Source/Programs/UnrealVS/UnrealVS.Shared/GenerateProjectFiles.cs
+++ b/Engine/Source/Programs/UnrealVS/UnrealVS.Shared/GenerateProjectFiles.cs
@@ -150,7 +150,11 @@ namespace UnrealVS
 			if (UnrealVSPackage.Instance.IsUESolutionLoaded)
 			{
 				// We expect "GenerateProjectFiles.bat" to live in the same directory as the solution
-				return Path.Combine(Path.GetDirectoryName(UnrealVSPackage.Instance.SolutionFilepath), "GenerateProjectFiles.bat");
+				string BatchPath = Path.Combine(Path.GetDirectoryName(UnrealVSPackage.Instance.SolutionFilepath), "GenerateProjectFiles.bat");
+				if (File.Exists(BatchPath))
+				{
+					return BatchPath;
+				}
 			}
 			return null;
 		}
diff --git a/Engine/Source/Programs/UnrealVS/UnrealVS.Shared/UnrealVSPackage.cs b/Engine/Source/Programs/UnrealVS/UnrealVS.Shared/UnrealVSPackage.cs
--- a/Engine/Source/Programs/UnrealVS/UnrealVS.Shared/UnrealVSPackage.cs
+++ b/Engine/Source/Programs/UnrealVS/UnrealVS.Shared/UnrealVSPackage.cs
@@ -79,7 +79,7 @@ namespace UnrealVS
 	{
 		/** Constants */
 
-		private const string VersionString = "v1.60";
+		private const string VersionString = "v1.75";
 		private const string UnrealSolutionFileNamePrefix = "UE";
 		private const string ExtensionName = "UnrealVS";
 		private const string CommandLineOptionKey = ExtensionName + "CommandLineMRU";
