UnrealVS: Save queue of files to be asynchronously checked out into a text file in the solution directory, and retry commands that don't succeed.

Also fix unrecognized options in launchSettings.json being cleared when using the command line editor.

diff --git a/Engine/Source/Programs/UnrealVS/ReleaseNotes.txt b/Engine/Source/Programs/UnrealVS/ReleaseNotes.txt
--- a/Engine/Source/Programs/UnrealVS/ReleaseNotes.txt
+++ b/Engine/Source/Programs/UnrealVS/ReleaseNotes.txt
@@ -1,3 +1,7 @@
+1.76
+* files scheduled for async checkout are now saved in a .p4checkout.txt file in the solution directory, to allow recovery if Perforce is down or VS terminates.
+* fixed some launch settings being deleted when using command line editor.
+
 1.72
 * fixed exception when opening project missing a referenced configuration
 * added exceptions to guard against code UI code not being accessed on the UI thread
diff --git a/Engine/Source/Programs/UnrealVS/UnrealVS.2019/Properties/AssemblyInfo.cs b/Engine/Source/Programs/UnrealVS/UnrealVS.2019/Properties/AssemblyInfo.cs
--- a/Engine/Source/Programs/UnrealVS/UnrealVS.2019/Properties/AssemblyInfo.cs
+++ b/Engine/Source/Programs/UnrealVS/UnrealVS.2019/Properties/AssemblyInfo.cs
@@ -31,8 +31,8 @@ using System.Runtime.InteropServices;
 // You can specify all the values or you can default the Revision and Build Numbers 
 // by using the '*' as shown below:
 
-[assembly: AssemblyVersion("1.75.0.0")]
-[assembly: AssemblyFileVersion("1.75.0.0")]
+[assembly: AssemblyVersion("1.76.0.0")]
+[assembly: AssemblyFileVersion("1.76.0.0")]
 
 
 
diff --git a/Engine/Source/Programs/UnrealVS/UnrealVS.2022/Properties/AssemblyInfo.cs b/Engine/Source/Programs/UnrealVS/UnrealVS.2022/Properties/AssemblyInfo.cs
--- a/Engine/Source/Programs/UnrealVS/UnrealVS.2022/Properties/AssemblyInfo.cs
+++ b/Engine/Source/Programs/UnrealVS/UnrealVS.2022/Properties/AssemblyInfo.cs
@@ -30,8 +30,8 @@ using System.Runtime.InteropServices;
 // You can specify all the values or you can default the Revision and Build Numbers 
 // by using the '*' as shown below:
 
-[assembly: AssemblyVersion("1.75.0.0")]
-[assembly: AssemblyFileVersion("1.75.0.0")]
+[assembly: AssemblyVersion("1.76.0.0")]
+[assembly: AssemblyFileVersion("1.76.0.0")]
 
 
 
diff --git a/Engine/Source/Programs/UnrealVS/UnrealVS.Shared/CommandLineEditor.cs b/Engine/Source/Programs/UnrealVS/UnrealVS.Shared/CommandLineEditor.cs
--- a/Engine/Source/Programs/UnrealVS/UnrealVS.Shared/CommandLineEditor.cs
+++ b/Engine/Source/Programs/UnrealVS/UnrealVS.Shared/CommandLineEditor.cs
@@ -30,9 +30,16 @@ namespace UnrealVS
 			{
 				[JsonPropertyName("commandName")]
 				public string CommandName { get; set; }
+
 				[JsonPropertyName("commandLineArgs")]
 				public string CommandLineArgs { get; set; }
+
+				[JsonExtensionData]
+				public Dictionary<string, JsonElement> ExtensionData { get; set; }
 			}
+
+			[JsonExtensionData]
+			public Dictionary<string, JsonElement> ExtensionData { get; set; }
 		}
 		/** constants */
 
diff --git a/Engine/Source/Programs/UnrealVS/UnrealVS.Shared/P4Commands.cs b/Engine/Source/Programs/UnrealVS/UnrealVS.Shared/P4Commands.cs
--- a/Engine/Source/Programs/UnrealVS/UnrealVS.Shared/P4Commands.cs
+++ b/Engine/Source/Programs/UnrealVS/UnrealVS.Shared/P4Commands.cs
@@ -5,19 +5,24 @@ using Microsoft.VisualStudio;
 using Microsoft.VisualStudio.Shell;
 using Microsoft.VisualStudio.Shell.Interop;
 using Microsoft.VisualStudio.Text.Differencing;
+using Microsoft.VisualStudio.Threading;
 using System;
 using System.Collections.Generic;
 using System.ComponentModel.Design;
 using System.Diagnostics;
 using System.IO;
+using System.Linq;
 using System.Text;
 using System.Text.RegularExpressions;
 using System.Threading;
 using System.Threading.Tasks;
+using System.Windows;
 using Process = System.Diagnostics.Process;
 
 namespace UnrealVS
 {
+	using Task = System.Threading.Tasks.Task;
+
 	internal class P4Commands : IDisposable
 	{
 		private const bool bPullWorkingDirectoryOn = true;
@@ -65,6 +70,10 @@ namespace UnrealVS
 
 		private List<CommandEvents> EventsForce = new List<CommandEvents>();
 
+		private readonly object CheckoutQueueLock = new object();
+		private readonly HashSet<string> CheckoutQueueFiles = new HashSet<string>(StringComparer.OrdinalIgnoreCase);
+		private JoinableTask CheckoutTask;
+
 		private class P4Command
 		{
 			public MenuCommand ButtonCommand;
@@ -85,7 +94,7 @@ namespace UnrealVS
 
 		private List<P4Command> P4CommandsList = new List<P4Command>();
 
-		private IntercepteSave Interceptor;
+		private InterceptSave Interceptor;
 
 		private bool IsSolutionLoaded()
 		{
@@ -96,6 +105,19 @@ namespace UnrealVS
 			return DTE.Solution.FileName.Length > 0;
 		}
 
+		private string GetCheckoutQueueFileName()
+		{
+			ThreadHelper.ThrowIfNotOnUIThread();
+
+			DTE DTE = UnrealVSPackage.Instance.DTE;
+			if (DTE.Solution.FileName.Length > 0)
+			{
+				return Path.Combine(Path.GetDirectoryName(DTE.Solution.FileName), ".p4checkout.txt");
+			}
+
+			return null;
+		}
+
 		private void OnQuickBuildSubMenuQuery(object sender, EventArgs e)
 		{
 			ThreadHelper.ThrowIfNotOnUIThread();
@@ -164,7 +186,7 @@ namespace UnrealVS
 			UpdateMenuOptions();
 
 			var runningDocumentTable = new RunningDocumentTable(UnrealVSPackage.Instance);
-			Interceptor = new IntercepteSave(UnrealVSPackage.Instance.DTE, runningDocumentTable, this);
+			Interceptor = new InterceptSave(UnrealVSPackage.Instance.DTE, runningDocumentTable, this);
 
 			runningDocumentTable.Advise(Interceptor);
 
@@ -176,6 +198,12 @@ namespace UnrealVS
 
 			// Update the menu visibility
 			UpdateMenuOptions();
+
+			string CheckoutQueueFile = GetCheckoutQueueFileName();
+			if (CheckoutQueueFile != null)
+			{
+				PulseCheckoutQueue(CheckoutQueueFile);
+			}
 		}
 
 		private void SolutionClosed()
@@ -656,13 +684,35 @@ namespace UnrealVS
 
 			if (UnrealVSPackage.Instance.OptionsPage.AllowAsyncP4Checkout)
 			{
+				string queueFile = GetCheckoutQueueFileName();
+				if (queueFile == null)
+				{
+					P4OutputPane.OutputString("Unable to get path to checkout queue file. No solution loaded?");
+					return;
+				}
+
+				// Add the file to the queue for checking out before modifying it. This allows us to recover in the case of a P4 outage or VS exit.
+				Logging.WriteLine($"Adding async checkout of {FileName} to {queueFile}");
+				lock (CheckoutQueueLock)
+				{
+					try
+					{
+						File.AppendAllLines(queueFile, new[] { FileName });
+					}
+					catch (Exception ex)
+					{
+						P4OutputPane.OutputString($"Unable to update {queueFile} ({ex.Message})");
+						return;
+					}
+				}
+
 				//P4OutputPane.OutputString($"Marking file as writeable: {FileName}");
 				// Mark the file as writeable
 				FileAttributes NewAttributes = File.GetAttributes(FileName) & ~FileAttributes.ReadOnly;
 				File.SetAttributes(FileName, NewAttributes);
 
-				// then send a fire and forget p4 edit command - does not lock the UI
-				_ = TryP4CommandAsync($"edit \"{FileName}\"");
+				// Start a background thread to update the queue
+				PulseCheckoutQueue(queueFile);
 			}
 			else
 			{
@@ -671,6 +721,94 @@ namespace UnrealVS
 			}
 		}
 
+		private void PulseCheckoutQueue(string QueueFile)
+		{
+			lock (CheckoutQueueLock)
+			{
+				CheckoutQueueFiles.Add(QueueFile);
+
+				if (CheckoutTask == null)
+				{
+					CheckoutTask = ThreadHelper.JoinableTaskFactory.RunAsync(UpdateCheckoutQueueAsync);
+				}
+			}
+		}
+
+		private async Task UpdateCheckoutQueueAsync()
+		{
+			for (; ; )
+			{
+				string QueueFile = null;
+				try
+				{
+					// Read the list of files that need checking out
+					string[] Lines;
+					lock (CheckoutQueueLock)
+					{
+						while (QueueFile == null || !File.Exists(QueueFile))
+						{
+							QueueFile = CheckoutQueueFiles.FirstOrDefault();
+							if (QueueFile == null)
+							{
+								CheckoutTask = null;
+								return;
+							}
+							CheckoutQueueFiles.Remove(QueueFile);
+						}
+						Lines = File.ReadAllLines(QueueFile);
+					}
+
+					// Open each file for edit
+					bool Pause = false;
+					foreach (string Line in Lines)
+					{
+						string TrimLine = Line.Trim();
+						if (TrimLine.Length > 0)
+						{
+							bool Success = await TryP4CommandAsync($"edit \"{TrimLine}\"");
+							lock (CheckoutQueueLock)
+							{
+								if (Success)
+								{
+									Logging.WriteLine($"Performed async checkout of {TrimLine}.");
+									lock (CheckoutQueueLock)
+									{
+										string[] NewLines = File.ReadAllLines(QueueFile);
+										File.WriteAllLines(QueueFile, NewLines.Where(x => !x.Equals(Line, StringComparison.Ordinal)));
+									}
+								}
+								else
+								{
+									// Set the UpdateCheckoutQueue flag again, to force another loop
+									Logging.WriteLine($"Unable to check out file {TrimLine}. Will retry.");
+									CheckoutQueueFiles.Add(QueueFile);
+									Pause = true;
+								}
+							}
+						}
+					}
+
+					// Check if we're done
+					if (Pause)
+					{
+						await Task.Delay(TimeSpan.FromSeconds(10.0));
+					}
+				}
+				catch (Exception ex)
+				{
+					Logging.WriteLine($"Unable to update checkout queue: {ex}");
+
+					if (QueueFile != null)
+					{
+						lock (CheckoutQueueLock)
+						{
+							CheckoutQueueFiles.Add(QueueFile);
+						}
+					}
+				}
+			}
+		}
+
 		private string ReadWorkingDirectorFromP4Config()
 		{
 			ThreadHelper.ThrowIfNotOnUIThread();
@@ -901,13 +1039,13 @@ namespace UnrealVS
 		}
 	}
 
-	internal class IntercepteSave : IVsRunningDocTableEvents3
+	internal class InterceptSave : IVsRunningDocTableEvents3
 	{
 		private readonly DTE DTE;
 		private readonly RunningDocumentTable RunningDocumentTable;
 		private readonly P4Commands P4Ops;
 
-		public IntercepteSave(DTE InDTE, RunningDocumentTable InRrunningDocumentTable, P4Commands InP4Ops)
+		public InterceptSave(DTE InDTE, RunningDocumentTable InRrunningDocumentTable, P4Commands InP4Ops)
 		{
 			DTE = InDTE;
 			RunningDocumentTable = InRrunningDocumentTable;
diff --git a/Engine/Source/Programs/UnrealVS/UnrealVS.Shared/UnrealVSPackage.cs b/Engine/Source/Programs/UnrealVS/UnrealVS.Shared/UnrealVSPackage.cs
--- a/Engine/Source/Programs/UnrealVS/UnrealVS.Shared/UnrealVSPackage.cs
+++ b/Engine/Source/Programs/UnrealVS/UnrealVS.Shared/UnrealVSPackage.cs
@@ -79,7 +79,7 @@ namespace UnrealVS
 	{
 		/** Constants */
 
-		private const string VersionString = "v1.75";
+		private const string VersionString = "v1.76";
 		private const string UnrealSolutionFileNamePrefix = "UE";
 		private const string ExtensionName = "UnrealVS";
 		private const string CommandLineOptionKey = ExtensionName + "CommandLineMRU";
