[UnrealVS version 1.77]
* FileBrowser - Added FileBrowser tool window that can be opened with "UnrealVs.FileBrowser". Press F1 to read help/features.
* UnrealVsPackage - Added events to for when projects are loaded/unloaded
* UnrealVsPackage - Added events to for when changing debugger mode
* UnrealVsPackage - Added parameter "last tick" to Tick in order to be able to close things before ticker stops
* UnrealVsPackage - Added code to include folder name to UE5 in main window and solution title. Optionally disabled in options
* Utils - Added code to set solution title and also prevent visual studio from setting solution title back to original value
* CompileSingleFile - Added command for compiling module that file is in (this file should probably be renamed to CompileSingleFileOrModule.cs)

diff --git a/Engine/Source/Programs/UnrealVS/ReleaseNotes.txt b/Engine/Source/Programs/UnrealVS/ReleaseNotes.txt
--- a/Engine/Source/Programs/UnrealVS/ReleaseNotes.txt
+++ b/Engine/Source/Programs/UnrealVS/ReleaseNotes.txt
@@ -1,3 +1,7 @@
+1.77
+* added File Browser dialog (open with command UnrealVs.FileBrowser)
+* added so Visual Studio title/solution show folder name of UE5.sln when UE5.sln is opened (no other solutions)
+
 1.76
 * files scheduled for async checkout are now saved in a .p4checkout.txt file in the solution directory, to allow recovery if Perforce is down or VS terminates.
 * fixed some launch settings being deleted when using command line editor.
diff --git a/Engine/Source/Programs/UnrealVS/UnrealVS.2019/Properties/AssemblyInfo.cs b/Engine/Source/Programs/UnrealVS/UnrealVS.2019/Properties/AssemblyInfo.cs
--- a/Engine/Source/Programs/UnrealVS/UnrealVS.2019/Properties/AssemblyInfo.cs
+++ b/Engine/Source/Programs/UnrealVS/UnrealVS.2019/Properties/AssemblyInfo.cs
@@ -31,8 +31,8 @@ using System.Runtime.InteropServices;
 // You can specify all the values or you can default the Revision and Build Numbers 
 // by using the '*' as shown below:
 
-[assembly: AssemblyVersion("1.76.0.0")]
-[assembly: AssemblyFileVersion("1.76.0.0")]
+[assembly: AssemblyVersion("1.77.0.0")]
+[assembly: AssemblyFileVersion("1.77.0.0")]
 
 
 
diff --git a/Engine/Source/Programs/UnrealVS/UnrealVS.2022/Properties/AssemblyInfo.cs b/Engine/Source/Programs/UnrealVS/UnrealVS.2022/Properties/AssemblyInfo.cs
--- a/Engine/Source/Programs/UnrealVS/UnrealVS.2022/Properties/AssemblyInfo.cs
+++ b/Engine/Source/Programs/UnrealVS/UnrealVS.2022/Properties/AssemblyInfo.cs
@@ -30,8 +30,8 @@ using System.Runtime.InteropServices;
 // You can specify all the values or you can default the Revision and Build Numbers 
 // by using the '*' as shown below:
 
-[assembly: AssemblyVersion("1.76.0.0")]
-[assembly: AssemblyFileVersion("1.76.0.0")]
+[assembly: AssemblyVersion("1.77.0.0")]
+[assembly: AssemblyFileVersion("1.77.0.0")]
 
 
 
diff --git a/Engine/Source/Programs/UnrealVS/UnrealVS.Shared/CompileSingleFile.cs b/Engine/Source/Programs/UnrealVS/UnrealVS.Shared/CompileSingleFile.cs
--- a/Engine/Source/Programs/UnrealVS/UnrealVS.Shared/CompileSingleFile.cs
+++ b/Engine/Source/Programs/UnrealVS/UnrealVS.Shared/CompileSingleFile.cs
@@ -10,6 +10,7 @@ using System.Diagnostics;
 using System.IO;
 using System.Windows.Forms;
 using System.Text;
+using System.Linq;
 
 namespace UnrealVS
 {
@@ -17,6 +18,7 @@ namespace UnrealVS
 	{
 		private const int CompileSingleFileButtonID = 0x1075;
 		private const int PreprocessSingleFileButtonID = 0x1076;
+		private const int CompileSingleModuleButtonID = 0x1077;
 		private const int UBTSubMenuID = 0x3103;
 		private string	  FileToCompileOriginalExt = "";
 
@@ -35,6 +37,10 @@ namespace UnrealVS
 			MenuCommand PreprocessSingleFileButtonCommand = new MenuCommand(new EventHandler(CompileSingleFileButtonHandler), CommandID2);
 			UnrealVSPackage.Instance.MenuCommandService.AddCommand(PreprocessSingleFileButtonCommand);
 
+			CommandID CommandID3 = new CommandID(GuidList.UnrealVSCmdSet, CompileSingleModuleButtonID);
+			MenuCommand CompileSingleModuleButtonCommand = new MenuCommand(new EventHandler(CompileSingleFileButtonHandler), CommandID3);
+			UnrealVSPackage.Instance.MenuCommandService.AddCommand(CompileSingleModuleButtonCommand);
+
 			// add sub menu for UBT commands
 			SubMenuCommand = new OleMenuCommand(null, new CommandID(GuidList.UnrealVSCmdSet, UBTSubMenuID));
 			UnrealVSPackage.Instance.MenuCommandService.AddCommand(SubMenuCommand);
@@ -51,9 +57,10 @@ namespace UnrealVS
 
 			MenuCommand SenderSubMenuCommand = (MenuCommand)Sender;
 
+			bool bIsFile = SenderSubMenuCommand.CommandID.ID != CompileSingleModuleButtonID;
 			bool PreprocessOnly = SenderSubMenuCommand.CommandID.ID == PreprocessSingleFileButtonID;
 
-			if (!TryCompileSingleFile(PreprocessOnly))
+			if (!TryCompileSingleFileOrModule(bIsFile, PreprocessOnly))
 			{
 				DTE DTE = UnrealVSPackage.Instance.DTE;
 				DTE.ExecuteCommand("Build.Compile");
@@ -74,7 +81,23 @@ namespace UnrealVS
 			}
 		}
 
-		bool TryCompileSingleFile(bool bPreProcessOnly)
+		string FindModuleForFile(string fileName, string rootDirectory)
+		{
+			string directory = Path.GetDirectoryName(fileName);
+			IEnumerable<string> buildFiles = Directory.EnumerateFiles(directory, "*.build.cs", SearchOption.TopDirectoryOnly);
+			string buildFile = buildFiles.FirstOrDefault();
+			if (buildFile != null)
+			{
+				return Path.GetFileName(buildFile.Substring(0, buildFile.LastIndexOf(".build.cs", StringComparison.OrdinalIgnoreCase)));
+			}
+			if (string.Equals(directory, rootDirectory, StringComparison.OrdinalIgnoreCase))
+			{
+				return null;
+			}
+			return FindModuleForFile(directory, rootDirectory);
+		}
+
+		bool TryCompileSingleFileOrModule(bool bIsFile, bool bPreProcessOnly)
 		{
 			ThreadHelper.ThrowIfNotOnUIThread();
 			DTE DTE = UnrealVSPackage.Instance.DTE;
@@ -152,10 +175,29 @@ namespace UnrealVS
 			// Check if the requested file is valid
 			string FileToCompile = DTE.ActiveDocument.FullName;
 			string FileToCompileExt = Path.GetExtension(FileToCompile);
-			if (!ValidExtensions.Contains(FileToCompileExt.ToLowerInvariant()))
+
+			string CompilingText;
+			string UBTArgument;
+			if (bIsFile)
 			{
-				MessageBox.Show($"Invalid file extension {FileToCompileExt} for single-file compile.", "Invalid Extension", MessageBoxButtons.OK);
-				return true;
+				if (!ValidExtensions.Contains(FileToCompileExt.ToLowerInvariant()))
+				{
+					MessageBox.Show($"Invalid file extension {FileToCompileExt} for single-file compile.", "Invalid Extension", MessageBoxButtons.OK);
+					return true;
+				}
+
+				CompilingText = FileToCompile;
+				UBTArgument = $"-singlefile=\"{FileToCompile}";
+			}
+			else
+			{
+				string ModuleName = FindModuleForFile(FileToCompile, Path.GetDirectoryName(DTE.Solution.FileName));
+				if (ModuleName == null)
+				{
+					MessageBox.Show($"Can't find module for for {FileToCompile} to compile.", "Invalid Module", MessageBoxButtons.OK);
+				}
+				CompilingText = ModuleName;
+				UBTArgument = $"-Module=\"{ModuleName}";
 			}
 
 			// If there's already a build in progress, don't let another one start
@@ -175,7 +217,7 @@ namespace UnrealVS
 			BuildOutputPane.Activate();
 			BuildOutputPane.Clear();
 			BuildOutputPane.OutputString($"1>------ Build started: Project: {StartupProject.Name}, Configuration: {ActiveConfiguration.ConfigurationName} {ActiveConfiguration.PlatformName} ------{Environment.NewLine}");
-			BuildOutputPane.OutputString($"1>  Compiling {FileToCompile}{Environment.NewLine}");
+			BuildOutputPane.OutputString($"1>  Compiling {CompilingText}{Environment.NewLine}");
 
 			// Set up event handlers 
 			DTE.Events.BuildEvents.OnBuildBegin += BuildEvents_OnBuildBegin;
@@ -207,7 +249,7 @@ namespace UnrealVS
 			// Spawn the new process
 			ChildProcess = new System.Diagnostics.Process();
 			ChildProcess.StartInfo.FileName = Path.Combine(Environment.SystemDirectory, "cmd.exe");
-			ChildProcess.StartInfo.Arguments = $"/C \"{BuildCommandLine} {PreProcess} -singlefile=\"{FileToCompile}\"\"";
+			ChildProcess.StartInfo.Arguments = $"/C \"{BuildCommandLine} {PreProcess} {UBTArgument}\"\"";
 			ChildProcess.StartInfo.WorkingDirectory = Path.GetDirectoryName(StartupProject.FullName);
 			ChildProcess.StartInfo.UseShellExecute = false;
 			ChildProcess.StartInfo.RedirectStandardOutput = true;
diff --git a/Engine/Source/Programs/UnrealVS/UnrealVS.Shared/FileBrowser.cs b/Engine/Source/Programs/UnrealVS/UnrealVS.Shared/FileBrowser.cs
new file mode 100644
--- /dev/null
+++ b/Engine/Source/Programs/UnrealVS/UnrealVS.Shared/FileBrowser.cs
@@ -0,0 +1,94 @@
+﻿// Copyright Epic Games, Inc. All Rights Reserved.
+
+using EnvDTE80;
+using Microsoft.VisualStudio.Shell;
+using Microsoft.VisualStudio.Shell.Interop;
+using System;
+using System.ComponentModel.Design;
+using System.Runtime.InteropServices;
+using System.Windows.Forms;
+
+namespace UnrealVS
+{
+	class FileBrowser
+	{
+		const int FileBrowserButtonID = 0x1337;
+
+		public FileBrowser()
+		{
+			// FileBrowserButton
+			var CommandID = new CommandID(GuidList.UnrealVSCmdSet, FileBrowserButtonID);
+			var FileBrowserButtonCommand = new MenuCommand(new EventHandler(FileBrowserButtonHandler), CommandID);
+			UnrealVSPackage.Instance.MenuCommandService.AddCommand(FileBrowserButtonCommand);
+
+			UnrealVSPackage.Instance.OnSolutionOpened += () => BrowserControl?.HandleSolutionChanged();
+			UnrealVSPackage.Instance.OnSolutionClosed += () => BrowserControl?.HandleSolutionChanged();
+			UnrealVSPackage.Instance.OnProjectLoaded += (p) => BrowserControl?.HandleSolutionChanged();
+			UnrealVSPackage.Instance.OnProjectUnloading += (p) => BrowserControl?.HandleSolutionChanged();
+			UnrealVSPackage.Instance.OnDocumentActivated += (d) => BrowserControl?.HandleDocumentActivated(d);
+
+		}
+
+		/// Called when 'FileBrowser' button is clicked
+		void FileBrowserButtonHandler(object Sender, EventArgs Args)
+		{
+			ThreadHelper.ThrowIfNotOnUIThread();
+
+			ToolWindowPane ToolWindow = UnrealVSPackage.Instance.FindToolWindow(typeof(FileBrowserWindow), 0, true);
+			if ((null == ToolWindow) || (null == ToolWindow.Frame))
+			{
+				throw new NotSupportedException(Resources.ToolWindowCreateError);
+			}
+			IVsWindowFrame ToolWindowFrame = (IVsWindowFrame)ToolWindow.Frame;
+			Microsoft.VisualStudio.ErrorHandler.ThrowOnFailure(ToolWindowFrame.Show());
+		}
+
+		FileBrowserWindowControl BrowserControl
+		{
+			get
+			{
+				ToolWindowPane ToolWindow = UnrealVSPackage.Instance.FindToolWindow(typeof(FileBrowserWindow), 0, true);
+				if ((null == ToolWindow) || (null == ToolWindow.Frame))
+					return null;
+				var fileBrowser = (FileBrowserWindow)ToolWindow;
+				return (FileBrowserWindowControl)fileBrowser.Content;
+			}
+		}
+	}
+
+	[Guid("b205a03c-566f-4dc3-9e2f-99e2df765f1d")]
+	public class FileBrowserWindow : ToolWindowPane
+	{
+		public FileBrowserWindow() : base(null)
+		{
+			this.Caption = "Unreal File Browser (F1 for help)";
+
+			this.Content = new FileBrowserWindowControl(this);
+		}
+
+		[global::System.Diagnostics.CodeAnalysis.SuppressMessageAttribute("", "VSTHRD010")]
+		protected override bool PreProcessMessage(ref Message m)
+        {
+			if (m.Msg == 0x0100)//WM_KEYDOWN)
+			{
+				if (m.WParam == new IntPtr(0x1B)) // VK_ESCAPE
+				{
+					((FileBrowserWindowControl)this.Content).HandleEscape();
+					return true;
+				}
+				else if (m.WParam == new IntPtr(0x70)) // VK_F1
+				{
+					((FileBrowserWindowControl)this.Content).HandleF1();
+					return true;
+				}
+				else if (m.WParam == new IntPtr(0x74)) // VK_F5
+                {
+					((FileBrowserWindowControl)this.Content).HandleF5();
+					return true;
+                }
+			}
+			return base.PreProcessMessage(ref m);
+        }
+    }
+
+}
diff --git a/Engine/Source/Programs/UnrealVS/UnrealVS.Shared/FileBrowserWindowControl.xaml b/Engine/Source/Programs/UnrealVS/UnrealVS.Shared/FileBrowserWindowControl.xaml
new file mode 100644
--- /dev/null
+++ b/Engine/Source/Programs/UnrealVS/UnrealVS.Shared/FileBrowserWindowControl.xaml
@@ -0,0 +1,198 @@
+﻿<UserControl x:Class="UnrealVS.FileBrowserWindowControl"
+             xmlns="http://schemas.microsoft.com/winfx/2006/xaml/presentation"
+             xmlns:x="http://schemas.microsoft.com/winfx/2006/xaml"
+             xmlns:mc="http://schemas.openxmlformats.org/markup-compatibility/2006"
+             xmlns:d="http://schemas.microsoft.com/expression/blend/2008"
+             xmlns:vsshell="clr-namespace:Microsoft.VisualStudio.Shell;assembly=Microsoft.VisualStudio.Shell.15.0"
+             Background="{DynamicResource {x:Static vsshell:VsBrushes.ToolWindowBackgroundKey}}"
+             Foreground="{DynamicResource {x:Static vsshell:VsBrushes.WindowTextKey}}"
+             mc:Ignorable="d"
+             d:DesignHeight="300" d:DesignWidth="500"
+             Name="FileBrowser">
+    <UserControl.Resources>
+        <Style TargetType="TextBox">
+            <Setter Property="Background" Value="Transparent"/>
+            <Setter Property="Foreground" Value="{DynamicResource {x:Static vsshell:VsBrushes.CommandBarTextActiveKey}}"/>
+        </Style>
+        <Style TargetType="ListBox">
+            <Setter Property="Background" Value="Transparent"/>
+            <Setter Property="BorderThickness" Value="1"/>
+            <Setter Property="Padding" Value="5" />
+            <Setter Property="BorderBrush" Value="{DynamicResource {x:Static vsshell:VsBrushes.ToolWindowBorderKey}}"/>
+            <Setter Property="Margin" Value="0,-6,0,0"/>
+            <Setter Property="FocusVisualStyle" Value="{x:Null}"/>
+        </Style>
+        <Style TargetType="TabItem">
+            <Setter Property="Background" Value="Transparent"/>
+            <Setter Property="IsTabStop" Value="False"/>
+            <Setter Property="Foreground" Value="{DynamicResource {x:Static vsshell:VsBrushes.WindowTextKey}}"/>
+            <Setter Property="FocusVisualStyle" Value="{StaticResource ListBoxItemVisualStyle}" />
+            <Setter Property="Template">
+                <Setter.Value>
+                    <ControlTemplate TargetType="TabItem">
+                        <Border Name="Border" BorderThickness="1,1,1,0" BorderBrush="{DynamicResource {x:Static vsshell:VsBrushes.ToolWindowBorderKey}}" CornerRadius="2,2,0,0" Margin="1,2,2,2">
+                            <ContentPresenter x:Name="ContentSite"
+                                    VerticalAlignment="Center"
+                                    HorizontalAlignment="Center"
+                                    ContentSource="Header"
+                                    Margin="10,3"/>
+                        </Border>
+                        <ControlTemplate.Triggers>
+                            <Trigger Property="IsSelected" Value="True">
+                                <Setter TargetName="Border" Property="Background" Value="{DynamicResource {x:Static vsshell:VsBrushes.WindowKey}}" />
+                            </Trigger>
+                            <Trigger Property="IsSelected" Value="False">
+                                <Setter TargetName="Border" Property="Background" Value="Transparent" />
+                                <Setter TargetName="ContentSite" Property="TextElement.Foreground" Value="{DynamicResource {x:Static vsshell:VsBrushes.FileTabInactiveTextKey}}" />
+                            </Trigger>
+                        </ControlTemplate.Triggers>
+                    </ControlTemplate>
+                </Setter.Value>
+            </Setter>
+        </Style>
+        <Style x:Key="ListBoxItemVisualStyle" TargetType="Control">
+        </Style>
+
+        <Style TargetType="ListBoxItem">
+            <Setter Property="Foreground" Value="{DynamicResource {x:Static vsshell:VsBrushes.CommandBarTextHoverKey}}"/>
+            <Setter Property="Template">
+                <Setter.Value>
+                    <ControlTemplate TargetType="{x:Type ListBoxItem}">
+                        <Border BorderBrush="Transparent" BorderThickness="0" Background="{TemplateBinding Background}">
+                            <ContentPresenter>
+                                <ContentPresenter.Resources>
+                                    <Style TargetType="TextBlock">
+                                        <Setter Property="Margin" Value="2"/>
+                                    </Style>
+                                </ContentPresenter.Resources>
+                            </ContentPresenter>
+                        </Border>
+                    </ControlTemplate>
+                </Setter.Value>
+            </Setter>
+            <Style.Triggers>
+                <MultiTrigger>
+                    <MultiTrigger.Conditions>
+                        <Condition Property="IsSelected" Value="true" />
+                    </MultiTrigger.Conditions>
+                    <Setter Property="Background" Value="{DynamicResource {x:Static vsshell:VsBrushes.ToolWindowTabMouseOverBackgroundBeginKey}}" />
+                    <Setter Property="FocusVisualStyle" Value="{StaticResource ListBoxItemVisualStyle}" />
+                </MultiTrigger>
+                <MultiDataTrigger>
+                    <MultiDataTrigger.Conditions>
+                        <Condition Binding="{Binding Path=Name}" Value="" />
+                    </MultiDataTrigger.Conditions>
+                    <Setter Property="Background" Value="{x:Null}" />
+                </MultiDataTrigger>
+            </Style.Triggers>
+        </Style>
+        <Style TargetType="Border">
+            <Setter Property="BorderBrush" Value="{DynamicResource {x:Static vsshell:VsBrushes.ToolWindowBorderKey}}"/>
+            <Setter Property="BorderThickness" Value="1"/>
+        </Style>
+    </UserControl.Resources>
+    <Grid>
+        <Grid.RowDefinitions>
+            <RowDefinition Height="*"/>
+        </Grid.RowDefinitions>
+        <Grid>
+            <Grid.RowDefinitions>
+                <RowDefinition Height="auto"/>
+                <RowDefinition Height="*"/>
+                <RowDefinition Height="auto"/>
+            </Grid.RowDefinitions>
+            <Border Margin="13,16,13,0" Padding="5" BorderBrush="{DynamicResource {x:Static vsshell:VsBrushes.ToolWindowBorderKey}}">
+                <Grid>
+                <Grid.ColumnDefinitions>
+                    <ColumnDefinition Width="auto"/>
+                    <ColumnDefinition Width="*"/>
+                </Grid.ColumnDefinitions>
+                    <TextBlock Text="Filter:" Padding="3"/>
+                    <TextBox Grid.Column="1" Name="FilterEditBox" CaretBrush="{DynamicResource {x:Static vsshell:VsBrushes.CommandBarTextActiveKey}}" Padding="3"/>
+                </Grid>
+            </Border>
+            <TabControl Name="FilesListTab" Grid.Row="1" IsTabStop="False" Background="Transparent" Margin="10" BorderBrush="Transparent">
+                <TabItem Header="All files" Name="AllFilesTab" Focusable="False">
+                    <ListBox Name="AllFilesListBox">
+                            <ListBox.ItemTemplate>
+                                <DataTemplate>
+                                    <TextBlock Text="{Binding Path=Name}" />
+                                </DataTemplate>
+                            </ListBox.ItemTemplate>
+                        </ListBox>
+                </TabItem>
+                <TabItem Header="Bookmarked files" Name="BookmarkedFilesTab" Focusable="False">
+                    <ListBox Name="BookmarkedFilesListBox">
+                        <ListBox.ItemTemplate>
+                            <DataTemplate>
+                                <TextBlock Text="{Binding Path=Name}" />
+                            </DataTemplate>
+                        </ListBox.ItemTemplate>
+                    </ListBox>
+                </TabItem>
+                <TabItem Header="Recent files" Name="RecentFilesTab" Focusable="False">
+                    <ListBox Name="RecentFilesListBox">
+                        <ListBox.ItemTemplate>
+                            <DataTemplate>
+                                <TextBlock Text="{Binding Path=Name}" />
+                            </DataTemplate>
+                        </ListBox.ItemTemplate>
+                    </ListBox>
+                </TabItem>
+            </TabControl>
+            <Border Grid.Row="2" Margin="13,0,13,16" Padding="5" BorderBrush="{DynamicResource {x:Static vsshell:VsBrushes.ToolWindowBorderKey}}">
+                <TextBox Name="StatusText" Text="" BorderThickness="0" IsReadOnly="True" IsTabStop="False"/>
+            </Border>
+        </Grid>
+        <Border Name="HelpDialog" BorderThickness="4" Width="auto" Height="auto" Visibility="Collapsed" Background="{DynamicResource {x:Static vsshell:VsBrushes.EditorExpansionFillKey}}" HorizontalAlignment="Center" VerticalAlignment="Center">
+            <StackPanel Margin="20">
+                <TextBlock Text="Help (Press escape to close this popup)"/>
+                <StackPanel Margin="6">
+                    <TextBlock Text="A file browser specialized for Unreal solutions. The browser is aware of .build.cs files"/>
+                    <TextBlock Text="To filter the file browser, simply type words, separated by space."/>
+                </StackPanel>
+                <TextBlock Text=""/>
+                <TextBlock Text="Keys"/>
+                <Grid>
+                    <Grid.ColumnDefinitions>
+                        <ColumnDefinition Width="auto"/>
+                        <ColumnDefinition Width="auto"/>
+                    </Grid.ColumnDefinitions>
+                    <StackPanel Grid.Column="0" Margin="6">
+                        <TextBlock Text="Escape"/>
+                        <TextBlock Text="Enter"/>
+                        <TextBlock Text="Up/Down"/>
+                        <TextBlock Text="Left/right"/>
+                        <TextBlock Text="Tab"/>
+                        <TextBlock Text="Insert"/>
+                        <TextBlock Text="Delete"/>
+                        <TextBlock Text="F1"/>
+                        <TextBlock Text="F5"/>
+                        <TextBlock Text="Ctrl-I"/>
+                    </StackPanel>
+                    <StackPanel Grid.Column="1" Margin="6">
+                        <TextBlock Text="- Close FileBrowser"/>
+                        <TextBlock Text="- Open selected file"/>
+                        <TextBlock Text="- Select file"/>
+                        <TextBlock Text="- Move between all files list and bookmarked files list"/>
+                        <TextBlock Text="- Toggle between filter edit box and file lists"/>
+                        <TextBlock Text="- Insert bookmarked file"/>
+                        <TextBlock Text="- Delete bookmarked file"/>
+                        <TextBlock Text="- This help"/>
+                        <TextBlock Text="- Refresh file lists with current visual studio solution"/>
+                        <TextBlock Text="- Create a #include &quot;file&quot; statement to clipboard and close file browser"/>
+                    </StackPanel>
+                </Grid>
+                <TextBlock Text=""/>
+                <TextBlock Text="Bookmarks"/>
+                <StackPanel Margin="6">
+                    <TextBlock Text="Can be used as a way to store favourite files and quickly navigate."/>
+                    <TextBlock Text="The list of bookmarks is solution specific and are stored next to the .sln file with extension .unrealvs"/>
+                </StackPanel>
+            </StackPanel>
+        </Border>
+        <Border Name="RefreshingDialog" Focusable="False" Width="300" Height="100" BorderThickness="4" Visibility="Collapsed" Background="{DynamicResource {x:Static vsshell:VsBrushes.EditorExpansionFillKey}}" HorizontalAlignment="Center" VerticalAlignment="Center">
+            <TextBlock Name="RefreshingText" Focusable="False" Text="" Width="auto" Height="auto" HorizontalAlignment="Center" VerticalAlignment="Center"/>
+        </Border>
+    </Grid>
+</UserControl>
diff --git a/Engine/Source/Programs/UnrealVS/UnrealVS.Shared/FileBrowserWindowControl.xaml.cs b/Engine/Source/Programs/UnrealVS/UnrealVS.Shared/FileBrowserWindowControl.xaml.cs
new file mode 100644
--- /dev/null
+++ b/Engine/Source/Programs/UnrealVS/UnrealVS.Shared/FileBrowserWindowControl.xaml.cs
@@ -0,0 +1,988 @@
+﻿using EnvDTE;
+using Microsoft.VisualStudio.Shell;
+using Microsoft.VisualStudio.Shell.Interop;
+using System;
+using System.Collections.Generic;
+using System.IO;
+using System.Text;
+using System.Text.Json;
+using System.Text.Json.Serialization;
+using System.Windows;
+using System.Windows.Controls;
+using System.Windows.Data;
+using System.Windows.Input;
+using System.Windows.Media;
+using System.Windows.Threading;
+
+namespace UnrealVS
+{
+    [global::System.Diagnostics.CodeAnalysis.SuppressMessageAttribute("", "VSTHRD010")]
+    public partial class FileBrowserWindowControl : UserControl
+    {
+        FileBrowserWindow Window;
+        bool ClearFilter;
+        bool IsRefreshing;
+        ListBox ActiveFilesListBox;
+
+        public class FileItem
+        {
+            public string Name { get; set; }
+            public string File { get; set; }
+            public string Project { get; set; }
+            public bool Selectable { get; set; }
+        }
+
+		String UnrealVsFileName;
+
+		List<FileItem> AllFileItems;
+        List<FileItem> BookmarkedFileItems;
+        List<FileItem> RecentFileItems;
+
+        Dictionary<string, FileItem> AllFileItemsLookup;
+
+		public FileBrowserWindowControl(FileBrowserWindow window)
+        {
+            Window = window;
+            this.InitializeComponent();
+
+            AllFileItems = new List<FileItem>();
+            BookmarkedFileItems = new List<FileItem>();
+            RecentFileItems = new List<FileItem>();
+            IsVisibleChanged += FileBrowserWindowControl_IsVisibleChanged;
+            foreach (var listView in new[] { AllFilesListBox, BookmarkedFilesListBox, RecentFilesListBox })
+            {
+                listView.ItemsSource = new List<FileItem>();
+                listView.KeyDown += FileListView_KeyDown;
+                listView.PreviewKeyDown += FileListView_PreviewKeyDown;
+                listView.SelectionChanged += FileListView_SelectionChanged;
+                listView.GotFocus += FileListView_GotFocus;
+                listView.MouseDoubleClick += ListView_MouseDoubleClick;
+                listView.SelectionMode = SelectionMode.Extended;
+            }
+            FilterEditBox.TextChanged += FilterEditBox_TextChanged;
+            FilterEditBox.KeyDown += FilterEditBox_KeyDown;
+            FilterEditBox.PreviewKeyDown += FilterEditBox_PreviewKeyDown;
+            FilesListTab.SelectionChanged += FilesListTab_SelectionChanged;
+            FilesListTab.KeyDown += FilesListTab_KeyDown;
+            FilesListTab.SelectedIndex = 0;
+
+			ActiveFilesListBox = AllFilesListBox;
+        }
+
+        private void FilesListTab_KeyDown(object sender, KeyEventArgs e)
+        {
+            /*
+            if (e.Key == Key.Right && FilesListTab.SelectedIndex < 2)
+                FilesListTab.SelectedIndex = (int)FilesListTab.SelectedIndex + 1;
+            else if (e.Key == Key.Left && FilesListTab.SelectedIndex > 0)
+                --FilesListTab.SelectedIndex;
+            */
+        }
+
+        private void FilesListTab_SelectionChanged(object sender, SelectionChangedEventArgs e)
+        {
+            switch (FilesListTab.SelectedIndex)
+            {
+                case 0:
+                    if (ActiveFilesListBox == AllFilesListBox)
+                        return;
+                    ActiveFilesListBox = AllFilesListBox;
+                    break;
+                case 1:
+                    if (ActiveFilesListBox == BookmarkedFilesListBox)
+                        return;
+                    ActiveFilesListBox = BookmarkedFilesListBox;
+                    break;
+				case 2:
+					if (ActiveFilesListBox == RecentFilesListBox)
+						return;
+					ActiveFilesListBox = RecentFilesListBox;
+					break;
+			}
+			FilesListTab.UpdateLayout();
+            RefreshStatusBox();
+            AsyncFocusSelectedItem();
+        }
+
+        private void FilterEditBox_KeyDown(object sender, KeyEventArgs e)
+        {
+            if (e.Key == Key.Tab)
+            {
+                e.Handled = true;
+                FocusListViewItem(ActiveFilesListBox, ActiveFilesListBox.SelectedIndex);
+                return;
+            }
+        }
+        private void FilterEditBox_PreviewKeyDown(object sender, KeyEventArgs e)
+        {
+            var listBox = ActiveFilesListBox;
+            var getSelectedItem = new Func<FileItem>(() =>
+            {
+                if (listBox.SelectedItems.Count == 0)
+                    return (FileItem)listBox.Items[0];
+                else
+                    return (FileItem)listBox.SelectedItems[listBox.SelectedItems.Count - 1];
+            });
+
+            var moveFunc = new Action<int>((int steps) =>
+            {
+                FileItem selectedItem = getSelectedItem();
+                var index = listBox.Items.IndexOf(selectedItem);
+                var newIndex = Math.Min(Math.Max(0, index + steps), listBox.Items.Count - 1);
+                if (newIndex == index)
+                    return;
+                selectedItem = (FileItem)listBox.Items[newIndex];
+                if (selectedItem == null)
+                    return;
+                listBox.SelectedItems.Clear();
+                listBox.SelectedItems.Add(selectedItem);
+                listBox.ScrollIntoView(selectedItem);
+            });
+
+            if (e.Key == Key.Down)
+            {
+                moveFunc(1);
+                return;
+            }
+            if (e.Key == Key.Up)
+            {
+                moveFunc(-1);
+                return;
+            }
+            if (e.Key == Key.Enter)
+            {
+                if (OpenSelectedFile())
+                    HidePanel();
+                return;
+            }
+            if (e.Key == Key.PageDown)
+            {
+                e.Handled = true;
+                var lastVisible = GetVisibleListViewElement(listBox, (int)listBox.ActualHeight - 22);
+                if (lastVisible == null)
+                    return;
+                var lastItem = listBox.ItemContainerGenerator.ItemFromContainer(lastVisible);
+                FileItem selectedItem = getSelectedItem();
+                if (lastItem != selectedItem)
+                {
+                    listBox.SelectedItems.Clear();
+                    listBox.SelectedItems.Add(lastItem);
+                    listBox.ScrollIntoView(lastItem);
+                }
+                else
+                    moveFunc(((int)(listBox.ActualHeight / lastVisible.ActualHeight))-1);
+                return;
+            }
+            if (e.Key == Key.PageUp)
+            {
+                e.Handled = true;
+                var firstVisible = GetVisibleListViewElement(listBox, 10);
+                if (firstVisible == null)
+                    return;
+                var firstItem = listBox.ItemContainerGenerator.ItemFromContainer(firstVisible);
+                FileItem selectedItem = getSelectedItem();
+                if (firstItem != selectedItem)
+                {
+                    listBox.SelectedItems.Clear();
+                    listBox.SelectedItems.Add(firstItem);
+                    listBox.ScrollIntoView(firstItem);
+                }
+                else
+                    moveFunc(1 - (int)(listBox.ActualHeight / firstVisible.ActualHeight));
+            }
+        }
+
+        private void FilterEditBox_TextChanged(object sender, TextChangedEventArgs e)
+        {
+            RefreshListViews();
+            //AllFilesListBox.SelectedIndex = 0;
+        }
+
+        private void FileListView_GotFocus(object sender, RoutedEventArgs e)
+        {
+            ActiveFilesListBox = (ListBox)e.Source;
+        }
+
+        private void FileListView_SelectionChanged(object sender, SelectionChangedEventArgs e)
+        {
+            RefreshStatusBox();
+        }
+
+        private void FileListView_KeyDown(object sender, KeyEventArgs e)
+        {
+            if ((Keyboard.Modifiers & ModifierKeys.Control) != 0)
+            {
+                if (e.Key == Key.I)
+                    if (CreateIncludePath())
+						HidePanel();
+				return;
+            }
+            if (e.Key == Key.Left)
+            {
+                if (ActiveFilesListBox != AllFilesListBox)
+                    --FilesListTab.SelectedIndex;
+                //if (ActiveFilesListBox == BookmarkedFilesListBox)
+                //    ToggleListView(BookmarkedFilesListBox, AllFilesListBox);
+                return;
+            }
+
+            if (e.Key == Key.Right)
+            {
+				if (ActiveFilesListBox != RecentFilesListBox)// && ((List<FileItem>)BookmarkedFilesListBox.ItemsSource).Count > 0)
+					++FilesListTab.SelectedIndex;
+				//if (ActiveFilesListBox == AllFilesListBox)// && ((List<FileItem>)BookmarkedFilesListBox.ItemsSource).Count > 0)
+                 //   FilesListTab.SelectedIndex = 1;
+                //if (ActiveFilesListBox == AllFilesListBox && ((List<FileItem>)BookmarkedFilesListBox.ItemsSource).Count > 0)
+               //     ToggleListView(AllFilesListBox, BookmarkedFilesListBox);
+                return;
+            }
+
+            if (e.Key == Key.Tab)
+            {
+                e.Handled = true;
+                FilterEditBox.Focus();
+                return;
+            }
+
+            if (e.Key == Key.Insert)
+            {
+                if (ActiveFilesListBox != AllFilesListBox)
+                    return;
+                bool BookmarksModified = false;
+                foreach (var i in ActiveFilesListBox.SelectedItems)
+                {
+                    var item = (FileItem)i;
+                    if (item == null || item.Name == "")
+                        continue;
+                    if (BookmarkedFileItems.Contains(item))
+                        continue;
+                    BookmarkedFileItems.Add(item);
+                    BookmarksModified = true;
+                }
+                if (!BookmarksModified)
+                    return;
+                RefreshListView(BookmarkedFilesListBox, BookmarkedFileItems);
+                //if (AllFilesListBox.SelectedIndex < AllFilesListBox.Items.Count - 1)
+                //    FocusListViewItem(AllFilesListBox, AllFilesListBox.SelectedIndex + 1);
+				SaveSolutionSettings();
+                return;
+            }
+
+            if (e.Key == Key.Delete)
+            {
+                if (ActiveFilesListBox != BookmarkedFilesListBox)
+                    return;
+                bool BookmarksModified = false;
+
+                var index = BookmarkedFilesListBox.SelectedIndex;
+                foreach (var i in ActiveFilesListBox.SelectedItems)
+                {
+                    var item = (FileItem)i;
+                    if (item == null || item.Name == "")
+                        continue;
+                    if (!BookmarkedFileItems.Remove(item))
+                        continue;
+                    BookmarksModified = true;
+                }
+                if (BookmarksModified)
+                {
+                    RefreshListView(BookmarkedFilesListBox, BookmarkedFileItems);
+                    if (index == BookmarkedFilesListBox.Items.Count)
+                    {
+                        if (index == 0)
+                            FilesListTab.SelectedIndex = 0;//ActiveFilesListBox = AllFilesListBox;
+                        else
+                            index -= 1;
+                    }
+                    BookmarkedFilesListBox.SelectedIndex = index;
+                    AsyncFocusSelectedItem();
+                    SaveSolutionSettings();
+                }
+                return;
+            }
+
+            if (e.Key == Key.Enter)
+            {
+                if (OpenSelectedFile())
+                    HidePanel();
+                return;
+            }
+
+            bool resetSelection = false;
+
+            if (e.Key == Key.Back)
+            {
+                string str = FilterEditBox.Text;
+                if (string.IsNullOrEmpty(str))
+                    return;
+                if (ClearFilter)
+                {
+                    resetSelection = true;
+                    ClearFilter = false;
+                    FilterEditBox.Text = "";
+                }
+                else
+                    FilterEditBox.Text = str.Substring(0, str.Length - 1);
+            }
+            else
+            {
+                char c = GetCharFromKey(e.Key);
+                if (c == 0)
+                    return;
+                if (ClearFilter)
+                {
+                    resetSelection = true;
+                    ClearFilter = false;
+                    FilterEditBox.Text = "";
+                }
+                FilterEditBox.Text += c;
+                FilterEditBox.SelectAll();
+            }
+
+            if (resetSelection)
+            {
+                AllFilesListBox.SelectedIndex = 0;
+                BookmarkedFilesListBox.SelectedIndex = 0;
+                RecentFilesListBox.SelectedIndex = 0;
+            }
+
+            AsyncFocusSelectedItem();
+
+            e.Handled = true;
+        }
+
+        private void ListView_MouseDoubleClick(object sender, MouseButtonEventArgs e)
+        {
+            if (OpenSelectedFile())
+                HidePanel();
+        }
+
+        private bool OpenSelectedFile()
+        {
+            var toOpen = new List<string>();
+            foreach (var i in ActiveFilesListBox.SelectedItems)
+            {
+                var item = (FileItem)i;
+                if (item == null || item.Name == "")
+                    continue;
+                toOpen.Add(item.File);
+            }
+            ThreadHelper.JoinableTaskFactory.Run(async delegate
+            {
+                await ThreadHelper.JoinableTaskFactory.SwitchToMainThreadAsync();
+                foreach (var f in toOpen)
+                    UnrealVSPackage.Instance.DTE.ExecuteCommand("File.OpenFile", "\"" + f + "\"");
+            });
+            return toOpen.Count != 0;
+        }
+
+        private bool CreateIncludePath()
+        {
+            string includeStrings = "";
+            foreach (var i in ActiveFilesListBox.SelectedItems)
+            {
+                var item = (FileItem)i;
+                if (item == null || item.Name == "")
+                    continue;
+                string includePath = item.File;
+                includePath = includePath.Replace("\\", "/");
+                string includePathLwr = includePath.ToLower();
+				var includeStartIndex = Math.Max(Math.Max(includePathLwr.LastIndexOf("/public/"), includePathLwr.LastIndexOf("/private/")), includePathLwr.LastIndexOf("/classes/"));
+                if (includeStartIndex != -1)
+                    includePath = includePath.Substring(includePath.IndexOf('/', includeStartIndex + 1) + 1);
+                includeStrings += $"#include \"{includePath}\"\r\n";
+            }
+            if (includeStrings == "")
+                return false;
+            Clipboard.SetText(includeStrings);
+            return true;
+		}
+
+		private void ToggleListView(ListBox from, ListBox to)
+        {
+            var firstVisibleFrom = GetVisibleListViewElement(from, 10);
+            if (firstVisibleFrom == null)
+                return;
+            var firstVisibleFromIndex = from.ItemContainerGenerator.IndexFromContainer(firstVisibleFrom);
+            var firstVisibleTo = GetVisibleListViewElement(to, 10);
+            if (firstVisibleTo == null)
+                return;
+            var firstVisibleToIndex = to.ItemContainerGenerator.IndexFromContainer(firstVisibleTo);
+            var toSelectIndex = firstVisibleToIndex + from.SelectedIndex - firstVisibleFromIndex;
+            if (toSelectIndex >= to.Items.Count)
+                toSelectIndex = to.Items.Count - 1;
+            FocusListViewItem(to, toSelectIndex);
+        }
+
+        private void FocusListViewItem(ListBox listView, int itemIndex)
+        {
+            if (itemIndex == -1)
+                itemIndex = 0;
+            listView.SelectedIndex = itemIndex;
+            if (listView.Items.Count > 0)
+                (listView.ItemContainerGenerator.ContainerFromIndex(itemIndex) as ListBoxItem)?.Focus();
+            else
+                listView.Focus();
+            RefreshStatusBox();
+        }
+
+        private void AsyncFocusSelectedItem()
+        {
+            double interval = 0.1;
+            var timer = new DispatcherTimer(DispatcherPriority.Normal);
+            timer.Tick += (s, e) =>
+            {
+                var items = (List<FileItem>)ActiveFilesListBox.ItemsSource;
+                if (items.Count == 0)
+                {
+                    ActiveFilesListBox.Focus();
+                    timer.Stop();
+                    return;
+                }
+                var item = (ActiveFilesListBox.ItemContainerGenerator.ContainerFromIndex(Math.Max(0, ActiveFilesListBox.SelectedIndex)) as ListBoxItem);
+                if (item == null)
+                {
+                    timer.Interval = TimeSpan.FromSeconds(interval);
+                    interval = Math.Min(interval + 0.1, 1.0);
+                    return;
+                }
+                item.Focus();
+                timer.Stop();
+            };
+            timer.Start();
+        }
+
+        private void RefreshListViews()
+        {
+            RefreshListView(AllFilesListBox, AllFileItems);
+			RefreshListView(BookmarkedFilesListBox, BookmarkedFileItems);
+			RefreshListView(RecentFilesListBox, RecentFileItems);
+		}
+
+        struct Score
+        {
+            public int Sorting;
+            public int Index;
+        }
+
+		private void RefreshListView(ListBox listView, List<FileItem> source)
+        {
+            var fileItems = (List<FileItem>)listView.ItemsSource;
+            fileItems.Clear();
+
+            string filterStr = FilterEditBox.Text;
+            var filter = filterStr.Split(new char[] { ' ' }, StringSplitOptions.RemoveEmptyEntries);
+
+            var selections = new HashSet<FileItem>();
+            foreach (var i in listView.SelectedItems)
+            {
+                var item = (FileItem)i;
+                if (item != null && item.Name != "")
+                    selections.Add(item);
+            }
+
+            var filterLen = filter.Length;
+            if (filterLen > 0)
+            {
+                var filterSortedItems = new List<ValueTuple<Score, FileItem>>();
+                var scoreArray = new ValueTuple<int, int>[5];
+                int ItemIndex = 0;
+                foreach (var item in source)
+                {
+                    int scoreArrayIndex = 0;
+                    for (int i = 0; i != filterLen; ++i)
+                    {
+                        var f = filter[i];
+                       int index = item.Name.IndexOf(f, System.StringComparison.OrdinalIgnoreCase);
+                        if (index == -1)
+                        {
+                            scoreArrayIndex = 0;
+                            break;
+                        }
+                        scoreArray[scoreArrayIndex].Item1 = index;
+                        scoreArray[scoreArrayIndex].Item2 = scoreArrayIndex;
+                        ++scoreArrayIndex;
+                        if (scoreArrayIndex == 5)
+                            break;
+                    }
+
+                    if (scoreArrayIndex == 0)
+                        continue;
+
+                    // We want to sort the ones that match the filter orders first and prioritize the ones that starts with the first filter entry
+                    int sortingScore = 0;
+                    Array.Sort(scoreArray, 0, scoreArrayIndex);
+                    for (int i = 0; i != scoreArrayIndex; ++i)
+                    {
+                        var Tuple = scoreArray[i];
+                        if (i == 0 && Tuple.Item1 == 0 && Tuple.Item2 == 0)
+                            sortingScore -= 100;
+                        sortingScore += (10 >> Tuple.Item2) * i + Tuple.Item1;
+                    }
+                    var s = new Score() { Sorting = sortingScore, Index = ItemIndex++ };
+                    filterSortedItems.Add(new ValueTuple<Score, FileItem>(s, item));
+                }
+
+                filterSortedItems.Sort((a, b) =>
+                {
+                    if (a.Item1.Sorting != b.Item1.Sorting)
+                        return a.Item1.Sorting - b.Item1.Sorting;
+                    return a.Item1.Index - b.Item1.Index;
+                });
+                foreach (var tuple in filterSortedItems)
+                    fileItems.Add(tuple.Item2);
+            }
+            else
+            {
+                fileItems.AddRange(source);
+            }
+
+            int fileCount = fileItems.Count;
+
+            // If empty we need to add one entry to make sure listview can be focused in windows
+            if (fileItems.Count == 0)
+                fileItems.Add(new FileItem() { Name = "" });
+
+            CollectionViewSource.GetDefaultView(fileItems)?.Refresh();
+
+            if (selections.Count != 0)
+            {
+                foreach (var item in fileItems)
+                {
+                    if (!selections.Contains(item))
+                        continue;
+                    listView.SelectedItems.Add(item);
+                    listView.ScrollIntoView(item);
+                }
+            }
+            
+            if (listView.SelectedItems.Count == 0)
+                listView.SelectedItems.Add(fileItems[0]);
+
+
+            if (listView == AllFilesListBox)
+                AllFilesTab.Header = $"All Files ({fileCount}/{source.Count})";
+            else if (listView == BookmarkedFilesListBox)
+                BookmarkedFilesTab.Header = $"Bookmarked Files ({fileCount}/{source.Count})";
+			else 
+				RecentFilesTab.Header = $"Recent Files ({fileCount}/{source.Count})";
+		}
+
+		private void RefreshStatusBox()
+        {
+            var items = ActiveFilesListBox.SelectedItems;
+            if (items.Count == 0)
+                StatusText.Text = "";
+            else if (items.Count == 1)
+            {
+                var item = (FileItem)items[0];
+                StatusText.Text = item.File;
+            }
+            else
+                StatusText.Text = "Multiple files selected";
+        }
+
+        private void FileListView_PreviewKeyDown(object sender, KeyEventArgs e)
+        {
+            if (e.Key != Key.Space)
+                return;
+            if (ClearFilter)
+            {
+                ClearFilter = false;
+                return;
+            }
+            FilterEditBox.Text += ' ';
+            e.Handled = true;
+        }
+
+        internal void HandleEscape()
+        {
+            if (HelpDialog.Visibility == Visibility.Visible)
+                HelpDialog.Visibility = Visibility.Collapsed;
+            else
+                HidePanel();
+        }
+
+        internal void HandleF1()
+        {
+            HelpDialog.Visibility = Visibility.Visible;
+        }
+
+        internal void HandleF5()
+        {
+            AsyncRefreshFileList();
+        }
+
+        internal void HandleSolutionChanged()
+        {
+            SaveSolutionSettings();
+            if (IsVisible)
+                AsyncRefreshFileList();
+            else
+                AllFileItems.Clear(); // Will trigger AsyncRefreshFileList when visible
+        }
+
+        internal void HandleDocumentActivated(Document Document)
+		{
+            var fileName = Document.FullName;
+            FileItem fileItem;
+            if (AllFileItemsLookup == null || !AllFileItemsLookup.TryGetValue(fileName, out fileItem))
+                return;
+
+            RecentFileItems.RemoveAll((item) => item == fileItem);
+            RecentFileItems.Insert(0, fileItem);
+
+            int RecentMaxCount = 100;
+            if (RecentFileItems.Count > RecentMaxCount)
+                RecentFileItems.RemoveRange(RecentMaxCount, RecentFileItems.Count - RecentMaxCount);
+
+			RefreshListView(RecentFilesListBox, RecentFileItems);
+		}
+
+		private void AsyncRefreshFileList()
+        {
+            if (IsRefreshing)
+                return;
+            IsRefreshing = true;
+            RefreshingText.Text = "Refreshing File Lists (0)";
+            RefreshingDialog.Visibility = Visibility.Visible;
+
+            DispatcherTimer timer = new DispatcherTimer(DispatcherPriority.Render);
+            timer.Interval = TimeSpan.FromSeconds(0.2);
+            var traverser = new SolutionTraverser();
+            timer.Tick += (s, e_) =>
+            {
+                var listItems = traverser.Update();
+                if (listItems == null)
+                {
+                    RefreshingText.Text = $"Refreshing File Lists ({traverser.HandledItemCount})";
+                    timer.Interval = TimeSpan.FromSeconds(0.01);
+                    return;
+                }
+                AllFileItems = listItems;
+
+				var lookup = new Dictionary<string, FileItem>();
+				foreach (var item in AllFileItems)
+					try { lookup.Add(item.File, item); } catch (Exception) { }
+                AllFileItemsLookup = lookup;
+
+				LoadSolutionSettings();
+
+                //ClearFilter = true;
+                RefreshListViews();
+
+                //AllFilesListBox.SelectedIndex = 0;
+				//BookmarkedFilesListBox.SelectedIndex = 0;
+				//RecentFilesListBox.SelectedIndex = 0;
+				AsyncFocusSelectedItem();
+
+                timer.Stop();
+                RefreshingDialog.Visibility = Visibility.Collapsed;
+                IsRefreshing = false;
+            };
+            timer.Start();
+        }
+
+        private void FileBrowserWindowControl_IsVisibleChanged(object sender, DependencyPropertyChangedEventArgs e)
+        {
+            if (!(bool)e.NewValue)
+                return;
+            FilterEditBox.SelectAll();
+            ClearFilter = true;
+            //FilesListTab.SelectedIndex = 0;
+
+            if (AllFileItems.Count == 0)
+            {
+                RefreshListViews(); // To make sure we get one invisible entry
+                AsyncRefreshFileList();
+            }
+
+            AsyncFocusSelectedItem();
+        }
+
+        private void HidePanel()
+        {
+            ThreadHelper.ThrowIfNotOnUIThread();
+            IVsWindowFrame ToolWindowFrame = (IVsWindowFrame)Window.Frame;
+            ToolWindowFrame.Hide();
+        }
+
+        private class FileBrowserSettings
+        {
+            public List<string> Bookmarks { get; set; } = new List<string>();
+			public List<string> Recents { get; set; } = new List<string>();
+		}
+
+
+		private void SaveSolutionSettings()
+        {
+            if (String.IsNullOrEmpty(UnrealVsFileName))
+                return;
+            using (var file = File.CreateText(UnrealVsFileName))
+            {
+                var settings = new FileBrowserSettings();
+				foreach (var item in BookmarkedFileItems)
+					settings.Bookmarks.Add(item.File);
+				foreach (var item in RecentFileItems)
+					settings.Recents.Add(item.File);
+				string json = JsonSerializer.Serialize(settings, new JsonSerializerOptions { WriteIndented = true });
+                file.Write(json);
+            }
+        }
+
+        private void LoadSolutionSettings()
+        {
+            var solutionFileName = UnrealVSPackage.Instance.DTE.Solution.FileName;
+            if (String.IsNullOrEmpty(solutionFileName))
+            {
+                BookmarkedFileItems.Clear();
+                RecentFileItems.Clear();
+                UnrealVsFileName = null;
+				return;
+            }
+            UnrealVsFileName = solutionFileName.Substring(0, solutionFileName.Length - 3) + "unrealvs";
+            if (!File.Exists(UnrealVsFileName))
+                return;
+            string json = File.ReadAllText(UnrealVsFileName);
+            if (String.IsNullOrEmpty(json))
+                return;
+            var settings = JsonSerializer.Deserialize<FileBrowserSettings>(json, new JsonSerializerOptions { PropertyNameCaseInsensitive = true });
+            if (settings == null)
+                return;
+
+
+            BookmarkedFileItems.Clear();
+            if (settings.Bookmarks.Count == 0 && settings.Recents.Count == 0)
+                return;
+
+            var fileItems = new[] { BookmarkedFileItems, RecentFileItems };
+            var settingItems = new[] { settings.Bookmarks, settings.Recents };
+			FileItem fileItem;
+
+            for (int i = 0; i != 2; ++i)
+            {
+                var added = new HashSet<FileItem>();
+                foreach (var item in settingItems[i])
+                    if (AllFileItemsLookup.TryGetValue(item, out fileItem))
+                        if (added.Add(fileItem))
+                            fileItems[i].Add(fileItem);
+            }
+        }
+
+        private static ListBoxItem GetVisibleListViewElement(ListBox listView, int y)
+        {
+            HitTestResult hitTest = VisualTreeHelper.HitTest(listView, new Point(10, y));
+            DependencyObject depObj = hitTest.VisualHit as DependencyObject;
+            if (depObj == null)
+                return null;
+            DependencyObject current = depObj;
+            while (current != null && current != listView)
+            {
+                ListBoxItem listBoxItem = current as ListBoxItem;
+                if (listBoxItem != null)
+                    return listBoxItem;
+                current = VisualTreeHelper.GetParent(current);
+            }
+            return null;
+        }
+        private static char GetCharFromKey(Key key)
+        {
+            char ch = '\0';
+
+            int virtualKey = KeyInterop.VirtualKeyFromKey(key);
+            byte[] keyboardState = new byte[256];
+            NativeMethods.GetKeyboardState(keyboardState);
+
+            uint scanCode = NativeMethods.MapVirtualKey((uint)virtualKey, NativeMethods.MapType.MAPVK_VK_TO_VSC);
+            StringBuilder stringBuilder = new StringBuilder(2);
+
+            int result = NativeMethods.ToUnicode((uint)virtualKey, scanCode, keyboardState, stringBuilder, stringBuilder.Capacity, 0);
+            switch (result)
+            {
+                case -1:
+                    break;
+                case 0:
+                    break;
+                case 1:
+                default:
+                    ch = stringBuilder[0];
+                    break;
+            }
+            return ch;
+        }
+
+        class StackItem
+        {
+            public ProjectItems Items;
+            public int Index;
+            public Project Project;
+        }
+
+        class SolutionTraverser
+        {
+            private SortedDictionary<string, FileItem> SortedFileItems = new SortedDictionary<string, FileItem>();
+            private Dictionary<string, Dictionary<string, FileItem>> CollidingFileItems = new Dictionary<string, Dictionary<string, FileItem>>();
+            private Dictionary<string, string> ModuleRoots = new Dictionary<string, string>();
+
+            private Stack<StackItem> ItemStack = new Stack<StackItem>();
+            private int ProjectIndex;
+            public int HandledItemCount { get; private set; }
+
+            public List<FileItem> Update()
+            {
+                Stack<StackItem> itemStack = ItemStack;
+                var projects = UnrealVSPackage.Instance.DTE.Solution.Projects;
+                var projectCount = projects.Count;
+
+                int traverseCounter = 0;
+
+                while (ProjectIndex < projectCount)
+                {
+                    var project = projects.Item(ProjectIndex + 1);
+
+                    StackItem stackItem;
+                    if (ItemStack.Count > 0)
+                        stackItem = ItemStack.Pop();
+                    else
+                        stackItem = new StackItem() { Items = project.ProjectItems, Project = project };
+
+                    while (true)
+                    {
+                        if (traverseCounter > 5000)
+                        {
+                            ItemStack.Push(stackItem);
+                            return null;
+                        }
+
+                        if (stackItem.Items == null || stackItem.Index == stackItem.Items.Count)
+                        {
+                            if (itemStack.Count == 0)
+                                break;
+                            stackItem = itemStack.Pop();
+                            ++stackItem.Index;
+                            continue;
+                        }
+
+                        var projectItem = stackItem.Items.Item(stackItem.Index + 1);
+                        ++traverseCounter;
+
+                        if (projectItem.FileCount != 0)
+                        {
+                            var file = projectItem.FileNames[1];
+                            if (file != null)
+                            {
+                                var name = projectItem.Name;
+                                if (name != file)
+                                    if (!file.EndsWith("\\")) // Skip folders
+                                    {
+                                        FileItem fileItem = new FileItem() { Name = name, File = file, Project = stackItem.Project.Name };
+                                        FileItem existingFileItem;
+                                        if (SortedFileItems.TryGetValue(name, out existingFileItem))
+                                        {
+                                            if (existingFileItem.File != file)
+                                            {
+                                                Dictionary<string, FileItem> colList;
+                                                if (!CollidingFileItems.TryGetValue(name, out colList))
+                                                {
+                                                    colList = new Dictionary<string, FileItem>();
+                                                    colList.Add(existingFileItem.File, existingFileItem);
+                                                    CollidingFileItems[name] = colList;
+                                                }
+                                                if (!colList.ContainsKey(fileItem.File))
+                                                    colList.Add(fileItem.File, fileItem);
+                                            }
+                                        }
+                                        else
+                                            SortedFileItems.Add(name, fileItem);
+
+                                        if (file.EndsWith(".build.cs", StringComparison.OrdinalIgnoreCase))
+                                        {
+                                            var lastBackslash = file.LastIndexOf('\\');
+                                            var moduleDir = file.Substring(0, lastBackslash);
+                                            var moduleName = file.Substring(lastBackslash + 1, file.Length - lastBackslash - 10);
+                                            if (!ModuleRoots.ContainsKey(moduleDir))
+                                                ModuleRoots.Add(moduleDir, moduleName);
+                                            //else if (ModuleRoots[moduleDir] != moduleName)
+                                            //    Logging.WriteLine(file);
+                                        }
+
+                                        ++HandledItemCount;
+                                    }
+                            }
+                        }
+
+                        if (projectItem.SubProject != null)
+                        {
+                            itemStack.Push(stackItem);
+                            stackItem = new StackItem() { Items = projectItem.SubProject.ProjectItems, Project = projectItem.SubProject };
+                            continue;
+                        }
+
+                        if (projectItem.ProjectItems != null)
+                        {
+                            itemStack.Push(stackItem);
+                            stackItem = new StackItem() { Items = projectItem.ProjectItems, Project = stackItem.Project };
+                            continue;
+                        }
+
+                        ++stackItem.Index;
+                    }
+
+                    ++ProjectIndex;
+                }
+
+                // We need to rename all collisions
+                foreach (var kv in CollidingFileItems)
+                {
+                    SortedFileItems.Remove(kv.Key);
+                    var usedParents = new Dictionary<string, int>();
+
+                    foreach (var itemKv in kv.Value)
+                    {
+                        var item = itemKv.Value;
+                        string parent = item.Project;
+                        var path = item.File;
+                        int lastBackslash = path.LastIndexOf('\\');
+                        while (lastBackslash != -1)
+                        {
+                            path = path.Substring(0, lastBackslash);
+                            string moduleName;
+                            if (ModuleRoots.TryGetValue(path, out moduleName))
+                            {
+                                parent = moduleName;
+                                break;
+                            }
+                            lastBackslash = path.LastIndexOf('\\');
+                        }
+
+                        string sortName = item.Name + $" ({parent})";
+
+                        int counter = 0;
+                        if (usedParents.TryGetValue(parent, out counter))
+                        {
+                            if (counter == 0)
+                            {
+                                var alreadyAddedItem = SortedFileItems[sortName];
+                                SortedFileItems.Remove(sortName);
+                                alreadyAddedItem.Name = item.Name + $" ({parent} 0)";
+                                SortedFileItems.Add(alreadyAddedItem.Name, alreadyAddedItem);
+                            }
+                            ++counter;
+                            usedParents[parent] = counter;
+                            item.Name = item.Name + $" ({parent} {counter})";
+                            SortedFileItems.Add(item.Name, item);
+                        }
+                        else
+                        {
+                            usedParents.Add(parent, 0);
+                            item.Name = sortName;
+                            SortedFileItems.Add(item.Name, item);
+                        }
+                    }
+                }
+
+                return new List<FileItem>(SortedFileItems.Values);
+            }
+        }
+    }
+}
\ No newline at end of file
diff --git a/Engine/Source/Programs/UnrealVS/UnrealVS.Shared/UnrealVSOptions.cs b/Engine/Source/Programs/UnrealVS/UnrealVS.Shared/UnrealVSOptions.cs
--- a/Engine/Source/Programs/UnrealVS/UnrealVS.Shared/UnrealVSOptions.cs
+++ b/Engine/Source/Programs/UnrealVS/UnrealVS.Shared/UnrealVSOptions.cs
@@ -52,6 +52,11 @@ namespace UnrealVS
 		[Description("Reconcile will attempt to resolve these extensions (including wildcards)")]
 		public string ReconcileExtensions { get; set; } = "c*;h*;ini;uproject;uplugin";
 
+		[Category("General")]
+		[DisplayName("Prefix UE5 solution name with folder")]
+		[Description("Will include name of folder that UE5.sln is located in to Visual Studio window title.")]
+		public bool IncludeFolderInUE5SolutionName { get; set; } = true;
+
 		protected override void OnApply(PageApplyEventArgs e)
 		{
 			base.OnApply(e);
diff --git a/Engine/Source/Programs/UnrealVS/UnrealVS.Shared/UnrealVSPackage.cs b/Engine/Source/Programs/UnrealVS/UnrealVS.Shared/UnrealVSPackage.cs
--- a/Engine/Source/Programs/UnrealVS/UnrealVS.Shared/UnrealVSPackage.cs
+++ b/Engine/Source/Programs/UnrealVS/UnrealVS.Shared/UnrealVSPackage.cs
@@ -6,7 +6,6 @@ using Microsoft.VisualStudio;
 using Microsoft.VisualStudio.OLE.Interop;
 using Microsoft.VisualStudio.Shell;
 using Microsoft.VisualStudio.Shell.Interop;
-using Microsoft.VisualStudio.Editor;
 using System;
 using System.Collections.Generic;
 using System.ComponentModel.Design;
@@ -15,16 +14,11 @@ using System.IO;
 using System.Linq;
 using System.Runtime.InteropServices;
 using System.Threading;
-using Microsoft.VisualStudio.TextManager.Interop;
 using Thread = System.Threading.Thread;
-using System.Windows.Forms;
 using Microsoft.VisualStudio.ComponentModelHost;
-using Microsoft.VisualStudio.Text;
 using Microsoft.VisualStudio.Text.Editor;
-using IOleServiceProvider = Microsoft.VisualStudio.OLE.Interop.IServiceProvider;
-using Microsoft.VisualStudio.Utilities;
-using Microsoft.VisualStudio.Text.Projection;
 using Microsoft;
+using System.Text;
 
 namespace UnrealVS
 {
@@ -64,6 +58,7 @@ namespace UnrealVS
 	// This attribute registers an options page for the package.
 	[ProvideOptionPage(typeof(UnrealVsOptions), ExtensionName, "General", 101, 102, true)]
 	[ProvideSolutionProperties(GuidList.UnrealVSPackageString)]
+	[ProvideToolWindow(typeof(FileBrowserWindow))]
 	/// <summary>
 	/// UnrealVSPackage implements Package abstract class.  This is the main class that is registered
 	/// with Visual Studio shell and serves as the entry point into our extension
@@ -75,11 +70,12 @@ namespace UnrealVS
 		IVsSelectionEvents,     // Allows us to be notified when the startup project has changed to a different project
 		IVsHierarchyEvents,     // Allows us to be notified when a hierarchy (the startup project) has had properties changed
 		IVsPersistSolutionProps, // Allows us to read props added to the solution to determine if the solution is a unreal solution
+		IVsDebuggerEvents,
 		IDisposable
 	{
 		/** Constants */
 
-		private const string VersionString = "v1.76";
+		private const string VersionString = "v1.77";
 		private const string UnrealSolutionFileNamePrefix = "UE";
 		private const string ExtensionName = "UnrealVS";
 		private const string CommandLineOptionKey = ExtensionName + "CommandLineMRU";
@@ -100,6 +96,14 @@ namespace UnrealVS
 		public delegate void OnProjectClosedDelegate(Project ClosedProject);
 		public event OnProjectClosedDelegate OnProjectClosed;
 
+		/// Called when a project is loaded in Visual Studio
+		public delegate void OnProjectLoadedDelegate(Project LoadedProject);
+		public event OnProjectLoadedDelegate OnProjectLoaded;
+
+		/// Called right before a project is unloaded in Visual Studio
+		public delegate void OnProjectUnloadingDelegate(Project UnloadedProject);
+		public event OnProjectUnloadingDelegate OnProjectUnloading;
+
 		/// Called when the startup project is edited in Visual Studio
 		public delegate void OnStartupProjectPropertyChangedDelegate(UInt32 itemid, Int32 propid, UInt32 flags);
 		public event OnStartupProjectPropertyChangedDelegate OnStartupProjectPropertyChanged;
@@ -129,6 +133,9 @@ namespace UnrealVS
 		public delegate void OnUIContextChangedDelegate(uint CmdUICookie, bool bActive);
 		public event OnUIContextChangedDelegate OnUIContextChanged;
 
+		public delegate void OnDocumentActivatedDelegate(Document Document);
+		public event OnDocumentActivatedDelegate OnDocumentActivated;
+
 		/** Public Fields & Properties */
 
 		/// Returns singleton instance of UnrealVSPackage
@@ -145,6 +152,10 @@ namespace UnrealVS
 		/// our singleton instance
 		public IVsSolutionBuildManager2 SolutionBuildManager { get; private set; }
 
+		/// Visual Studio solution build manager interface.  This is used to change the active startup
+		/// Project, among other things.  We expose public access to the solution build manager through
+		/// our singleton instance
+		public IVsDebugger Debugger { get; private set; }
 
 		/// Visual Studio solution "manager" interface.  We register with this to receive events
 		/// about projects being added and such.  This needs to be cleaned up at shutdown.
@@ -235,6 +246,11 @@ namespace UnrealVS
 			Assumes.Present(SolutionBuildManager);
 			SolutionBuildManager.AdviseUpdateSolutionEvents(this, out UpdateSolutionEventsHandle);
 
+			// Get debugger
+			Debugger = await GetServiceAsync(typeof(SVsShellDebugger)) as IVsDebugger;
+			Assumes.Present(Debugger);
+			Debugger.AdviseDebuggerEvents(this, out DebuggerEventsHandle);
+
 			// Create our command-line editor
 			CommandLineEditor.Initialize();
 
@@ -259,6 +275,9 @@ namespace UnrealVS
 			// Create 'Perforce menu' instance
 			P4CommandsGroup = new P4Commands();
 
+			// Create 'FileBrowser' instance
+			FileBrowser = new FileBrowser();
+
 			// Call parent implementation
 			base.Initialize();
 
@@ -267,10 +286,19 @@ namespace UnrealVS
 				StartTicker();
 			}
 
+			DTE2.Events.WindowEvents.WindowActivated += WindowEvents_WindowActivated;
+		}
+
+		private void WindowEvents_WindowActivated(Window GotFocus, Window LostFocus)
+		{
+			ThreadHelper.ThrowIfNotOnUIThread();
+			if (GotFocus.Document != null)
+				OnDocumentActivated(GotFocus.Document);
 		}
 
 		private void StartTicker()
 		{
+			bRefreshTitleInTick = true;
 			// Create a "ticker" on a background thread that ticks the package on the UI thread
 			Interlocked.Exchange(ref bCancelTicker, 0);
 			Ticker = new Thread(TickAsyncMain)
@@ -295,17 +323,19 @@ namespace UnrealVS
 		{
 			try
 			{
+				bool bLastTick = false;
 				while (true)
 				{
-					if (bCancelTicker != 0) return;
-
 					ThreadHelper.JoinableTaskFactory.Run(async () =>
 					{
 						await ThreadHelper.JoinableTaskFactory.SwitchToMainThreadAsync();
-						Tick();
+						if (bCancelTicker != 0)
+							bLastTick = true;
+						Tick(bLastTick);
 					});
 
-					if (bCancelTicker != 0) return;
+					if (bLastTick)
+						return;
 					Thread.Sleep(TickPeriod);
 				}
 			}
@@ -317,12 +347,57 @@ namespace UnrealVS
 		/// <summary>
 		/// Tick function on main UI thread
 		/// </summary>
-		private void Tick()
+		private void Tick(bool bLastTick)
 		{
 			ThreadHelper.ThrowIfNotOnUIThread();
+
+			TickRefreshTitle(bLastTick);
+
 			BatchBuilder.Tick();
 		}
 
+		public void TickRefreshTitle(bool bLastTick)
+		{
+			ThreadHelper.ThrowIfNotOnUIThread();
+
+			if (bLastTick)
+			{
+				Utils.SolutionTitle = null;
+				return;
+			}
+
+			if (bRefreshTitleInTick)
+			{
+				bRefreshTitleInTick = false;
+
+				if (!UnrealVSPackage.Instance.OptionsPage.IncludeFolderInUE5SolutionName)
+				{
+					return;
+				}
+
+				string SolutionFileName = UnrealVSPackage.Instance.DTE.Solution.FileName;
+				string SolutionName = Path.GetFileNameWithoutExtension(SolutionFileName);
+
+				// Only affect UE5 solution
+				if (SolutionName != "UE5")
+				{
+					return;
+				}
+
+				string FolderName = Path.GetFileName(Path.GetDirectoryName(SolutionFileName));
+				string Mode = "";
+				if (DbgMode == DBGMODE.DBGMODE_Run)
+					Mode = " (Running)";
+				if (DbgMode == DBGMODE.DBGMODE_Break)
+					Mode = " (Debugging)";
+				string NewTitle = FolderName + '\\' + SolutionName;
+
+				Utils.SolutionTitle = NewTitle;
+
+				Utils.MainWindowTitle = NewTitle + Mode + " - Microsoft Visual Studio";
+			}
+		}
+
 		/// <summary>
 		/// Implementation from IDisposable
 		/// </summary>
@@ -381,6 +456,11 @@ namespace UnrealVS
 			}
 			SelectionManager = null;
 
+			if (DebuggerEventsHandle != 0)
+				Debugger.UnadviseDebuggerEvents(DebuggerEventsHandle);
+			Debugger = null;
+
+
 			// No longer want update solution events
 			if (UpdateSolutionEventsHandle != 0)
 			{
@@ -625,6 +705,11 @@ namespace UnrealVS
 
 		int IVsSolutionEvents.OnAfterLoadProject(IVsHierarchy pStubHierarchy, IVsHierarchy pRealHierarchy)
 		{
+			ThreadHelper.ThrowIfNotOnUIThread();
+
+			var LoadedProject = Utils.HierarchyObjectToProject(pRealHierarchy);
+			if (LoadedProject != null && OnProjectLoaded != null)
+				OnProjectLoaded.Invoke(LoadedProject);
 			return VSConstants.S_OK;
 		}
 
@@ -687,6 +772,11 @@ namespace UnrealVS
 
 		int IVsSolutionEvents.OnBeforeUnloadProject(IVsHierarchy pRealHierarchy, IVsHierarchy pStubHierarchy)
 		{
+			ThreadHelper.ThrowIfNotOnUIThread();
+
+			var UnloadedProject = Utils.HierarchyObjectToProject(pRealHierarchy);
+			if (UnloadedProject != null && OnProjectUnloading != null)
+				OnProjectUnloading.Invoke(UnloadedProject);
 			return VSConstants.S_OK;
 		}
 
@@ -874,6 +964,14 @@ namespace UnrealVS
 
 		#endregion
 
+		int IVsDebuggerEvents.OnModeChange(DBGMODE dbgmodeNew)
+		{
+			DbgMode = dbgmodeNew;
+			bRefreshTitleInTick = true;
+			return VSConstants.S_OK;
+		}
+
+
 		// IVsUpdateSolutionEvents Interface
 
 		public int UpdateSolution_Begin(ref int pfCancelUpdate)
@@ -975,6 +1073,9 @@ namespace UnrealVS
 		/// Handle that we use at shutdown to unregister for events about solution build activity
 		private UInt32 UpdateSolutionEventsHandle;
 
+		/// Handle that we use at shutdown to unregister for events about solution build activity
+		private UInt32 DebuggerEventsHandle;
+
 		/// Handle that we use to unregister for events about startup project hierarchy activity
 		UInt32 ProjectHierarchyEventsHandle;
 
@@ -984,6 +1085,9 @@ namespace UnrealVS
 		/// BuildStartupProject feature
 		private BuildStartupProject BuildStartupProject;
 
+		/// FileBrowser feature
+		private FileBrowser FileBrowser;
+
 		/// CompileSingleFile feature
 		private CompileSingleFile CompileSingleFile;
 
@@ -1007,6 +1111,9 @@ namespace UnrealVS
 		private readonly List<string> LoadedProjectPaths = new List<string>();
 		private bool? _IsUESolutionLoaded;
 
+		private DBGMODE DbgMode = DBGMODE.DBGMODE_Design;
+		private bool bRefreshTitleInTick;
+
 		/// Obtains the DTE2 interface for this instance of VS from the RunningObjectTable
 		private static DTE2 GetDTE2ForCurrentInstance(DTE DTE)
 		{
@@ -1042,5 +1149,29 @@ namespace UnrealVS
 		/// ROT function in ole32.dll needed by GetDTE2ForCurrentInstance()
 		[DllImport("ole32.dll")]
 		internal static extern int GetRunningObjectTable(int reserved, out IRunningObjectTable prot);
+
+		public enum MapType : uint
+		{
+			MAPVK_VK_TO_VSC = 0x0,
+			MAPVK_VSC_TO_VK = 0x1,
+			MAPVK_VK_TO_CHAR = 0x2,
+			MAPVK_VSC_TO_VK_EX = 0x3,
+		}
+
+		[DllImport("user32.dll")]
+		public static extern int ToUnicode(
+			uint wVirtKey,
+			uint wScanCode,
+			byte[] lpKeyState,
+			[Out, MarshalAs(UnmanagedType.LPWStr, SizeParamIndex = 4)]
+			StringBuilder pwszBuff,
+			int cchBuff,
+			uint wFlags);
+
+		[DllImport("user32.dll")]
+		public static extern bool GetKeyboardState(byte[] lpKeyState);
+
+		[DllImport("user32.dll")]
+		public static extern uint MapVirtualKey(uint uCode, MapType uMapType);
 	}
 }
diff --git a/Engine/Source/Programs/UnrealVS/UnrealVS.Shared/Utils.cs b/Engine/Source/Programs/UnrealVS/UnrealVS.Shared/Utils.cs
--- a/Engine/Source/Programs/UnrealVS/UnrealVS.Shared/Utils.cs
+++ b/Engine/Source/Programs/UnrealVS/UnrealVS.Shared/Utils.cs
@@ -7,8 +7,11 @@ using Microsoft.VisualStudio.Shell;
 using Microsoft.VisualStudio.Shell.Interop;
 using System;
 using System.Collections.Generic;
+using System.ComponentModel;
 using System.IO;
 using System.Linq;
+using System.Windows.Controls;
+using System.Windows.Media;
 using System.Xml;
 
 namespace UnrealVS
@@ -721,7 +724,7 @@ namespace UnrealVS
 			}
 		}
 
-		private static string GetSolutionFolder()
+		public static string GetSolutionFolder()
 		{
 			ThreadHelper.ThrowIfNotOnUIThread();
 
@@ -733,6 +736,115 @@ namespace UnrealVS
 			return Path.GetDirectoryName(UnrealVSPackage.Instance.SolutionFilepath);
 		}
 
+		public static string SolutionTitle
+		{
+			set
+			{
+				if (!SolutionTextBlockSearched)
+				{
+					SolutionTextBlockSearched = true;
+
+					if (Utils.FindChild<System.Windows.DependencyObject>(System.Windows.Application.Current.MainWindow, "PART_SolutionNameTextBlock") is var textBlock)
+					{
+						SolutionTextBlock = Utils.FindChild<TextBlock>(textBlock, null);
+					}
+				}
+
+				if (SolutionTextBlock != null)
+				{
+					SolutionTextBlockText = value;
+					if (value != null)
+					{
+						SolutionTextBlockChanging = true;
+						SolutionTextBlock.Text = value;
+						SolutionTextBlockChanging = false;
+
+						if (!SolutionTextBlockTracked)
+						{
+							SolutionTextBlockTracked = true;
+							var dp = DependencyPropertyDescriptor.FromProperty(TextBlock.TextProperty, typeof(TextBlock));
+							dp.AddValueChanged(SolutionTextBlock, SolutionTextChanged);
+						}
+					}
+					else
+					{
+						if (SolutionTextBlockTracked)
+						{
+							SolutionTextBlockTracked = false;
+							var dp = DependencyPropertyDescriptor.FromProperty(TextBlock.TextProperty, typeof(TextBlock));
+							dp.RemoveValueChanged(SolutionTextBlock, SolutionTextChanged);
+						}
+						SolutionTextBlock.InvalidateVisual();
+					}
+				}
+			}
+		}
+
+		static void SolutionTextChanged(object sender, EventArgs e)
+		{
+			// Needed to prevent visual studio from changing text back to name from property
+			if (!SolutionTextBlockChanging && SolutionTextBlockText != null)
+			{
+				SolutionTextBlockChanging = true;
+				SolutionTextBlock.Text = SolutionTextBlockText;
+				SolutionTextBlockChanging = false;
+			}
+		}
+
+		public static string MainWindowTitle
+		{
+			get => System.Windows.Application.Current.MainWindow.Title;
+			set { System.Windows.Application.Current.MainWindow.Title = value; }
+		}
+
+		public static T FindChild<T>(System.Windows.DependencyObject Parent, string ChildName) where T : System.Windows.DependencyObject
+		{
+			if (Parent == null)
+			{
+				return null;
+			}
+
+			int ChildrenCount = VisualTreeHelper.GetChildrenCount(Parent);
+			for (int i = 0; i < ChildrenCount; i++)
+			{
+				var Child = VisualTreeHelper.GetChild(Parent, i);
+				if (ChildName != null)
+				{
+					// If the child's name is set for search
+					var FrameworkElement = Child as System.Windows.FrameworkElement;
+					if (FrameworkElement != null)
+					{
+						if (FrameworkElement.Name == ChildName)
+						{
+							// if the child's name is of the request name
+							return (T)Child;
+						}
+					}
+				}
+				else
+				{
+					if (Child is T TypedChild)
+					{
+						return TypedChild;
+					}
+				}
+
+				// recursively drill down the tree
+				T FoundChild = FindChild<T>(Child, ChildName);
+				if (FoundChild != null)
+				{
+					return FoundChild;
+				}
+			}
+
+			return null;
+		}
+
+		private static string SolutionTextBlockText;
+		private static TextBlock SolutionTextBlock;
+		private static bool SolutionTextBlockTracked;
+		private static bool SolutionTextBlockChanging;
+		private static bool SolutionTextBlockSearched;
 		private static string CachedUProjectRootFolder = string.Empty;
 		private static IEnumerable<string> CachedUProjectPaths = new string[0];
 		private static IDictionary<string, string> CachedUProjects = null;
