[UnrealVs]
* Added copyright comment at top of .cs file

diff --git a/Engine/Source/Programs/UnrealVS/UnrealVS.Shared/FileBrowserWindowControl.xaml.cs b/Engine/Source/Programs/UnrealVS/UnrealVS.Shared/FileBrowserWindowControl.xaml.cs
--- a/Engine/Source/Programs/UnrealVS/UnrealVS.Shared/FileBrowserWindowControl.xaml.cs
+++ b/Engine/Source/Programs/UnrealVS/UnrealVS.Shared/FileBrowserWindowControl.xaml.cs
@@ -1,4 +1,6 @@
-﻿using EnvDTE;
+﻿// Copyright Epic Games, Inc. All Rights Reserved.
+
+using EnvDTE;
 using Microsoft.VisualStudio.Shell;
 using Microsoft.VisualStudio.Shell.Interop;
 using System;
