UnrealVS: Fix race condition when performing async checkout.

Initial execution of background task occurs on thread that creates it, which currently holds checkout lock. That can complete before the caller assigns value to CheckoutTask.

diff --git a/Engine/Source/Programs/UnrealVS/ReleaseNotes.txt b/Engine/Source/Programs/UnrealVS/ReleaseNotes.txt
--- a/Engine/Source/Programs/UnrealVS/ReleaseNotes.txt
+++ b/Engine/Source/Programs/UnrealVS/ReleaseNotes.txt
@@ -1,3 +1,6 @@
+1.78
+* fixed race condition when starting async checkout
+
 1.77
 * added File Browser dialog (open with command UnrealVs.FileBrowser)
 * added so Visual Studio title/solution show folder name of UE5.sln when UE5.sln is opened (no other solutions)
diff --git a/Engine/Source/Programs/UnrealVS/UnrealVS.2019/Properties/AssemblyInfo.cs b/Engine/Source/Programs/UnrealVS/UnrealVS.2019/Properties/AssemblyInfo.cs
--- a/Engine/Source/Programs/UnrealVS/UnrealVS.2019/Properties/AssemblyInfo.cs
+++ b/Engine/Source/Programs/UnrealVS/UnrealVS.2019/Properties/AssemblyInfo.cs
@@ -31,8 +31,8 @@ using System.Runtime.InteropServices;
 // You can specify all the values or you can default the Revision and Build Numbers 
 // by using the '*' as shown below:
 
-[assembly: AssemblyVersion("1.77.0.0")]
-[assembly: AssemblyFileVersion("1.77.0.0")]
+[assembly: AssemblyVersion("1.78.0.0")]
+[assembly: AssemblyFileVersion("1.78.0.0")]
 
 
 
diff --git a/Engine/Source/Programs/UnrealVS/UnrealVS.2022/Properties/AssemblyInfo.cs b/Engine/Source/Programs/UnrealVS/UnrealVS.2022/Properties/AssemblyInfo.cs
--- a/Engine/Source/Programs/UnrealVS/UnrealVS.2022/Properties/AssemblyInfo.cs
+++ b/Engine/Source/Programs/UnrealVS/UnrealVS.2022/Properties/AssemblyInfo.cs
@@ -30,8 +30,8 @@ using System.Runtime.InteropServices;
 // You can specify all the values or you can default the Revision and Build Numbers 
 // by using the '*' as shown below:
 
-[assembly: AssemblyVersion("1.77.0.0")]
-[assembly: AssemblyFileVersion("1.77.0.0")]
+[assembly: AssemblyVersion("1.78.0.0")]
+[assembly: AssemblyFileVersion("1.78.0.0")]
 
 
 
diff --git a/Engine/Source/Programs/UnrealVS/UnrealVS.Shared/P4Commands.cs b/Engine/Source/Programs/UnrealVS/UnrealVS.Shared/P4Commands.cs
--- a/Engine/Source/Programs/UnrealVS/UnrealVS.Shared/P4Commands.cs
+++ b/Engine/Source/Programs/UnrealVS/UnrealVS.Shared/P4Commands.cs
@@ -190,7 +190,14 @@ namespace UnrealVS
 
 			runningDocumentTable.Advise(Interceptor);
 
+			string CheckoutQueueFile = GetCheckoutQueueFileName();
+			if (CheckoutQueueFile != null)
+			{
+				P4OutputPane.OutputString($"UnrealVS started. Using checkout queue: {CheckoutQueueFile}" + Environment.NewLine);
+				PulseCheckoutQueue(CheckoutQueueFile);
+			}
 		}
+
 		// Called when solutions are loaded or unloaded
 		private void SolutionOpened()
 		{
@@ -202,6 +209,7 @@ namespace UnrealVS
 			string CheckoutQueueFile = GetCheckoutQueueFileName();
 			if (CheckoutQueueFile != null)
 			{
+				P4OutputPane.OutputString($"Solution opened. Using checkout queue: {CheckoutQueueFile}" + Environment.NewLine);
 				PulseCheckoutQueue(CheckoutQueueFile);
 			}
 		}
@@ -678,7 +686,7 @@ namespace UnrealVS
 			if (!File.Exists(FileName) ||
  				(CheckOptionsFlag && !File.GetAttributes(FileName).HasFlag(FileAttributes.ReadOnly)))
 			{
-				P4OutputPane.OutputString($"already writeable: {FileName}");
+				P4OutputPane.OutputString($"already writeable: {FileName}" + Environment.NewLine);
 				return;
 			}
 
@@ -701,7 +709,7 @@ namespace UnrealVS
 					}
 					catch (Exception ex)
 					{
-						P4OutputPane.OutputString($"Unable to update {queueFile} ({ex.Message})");
+						P4OutputPane.OutputString($"Unable to update {queueFile} ({ex.Message})" + Environment.NewLine);
 						return;
 					}
 				}
@@ -723,15 +731,19 @@ namespace UnrealVS
 
 		private void PulseCheckoutQueue(string QueueFile)
 		{
+			TaskCompletionSource<bool> StartTask = null;
 			lock (CheckoutQueueLock)
 			{
 				CheckoutQueueFiles.Add(QueueFile);
 
 				if (CheckoutTask == null)
 				{
-					CheckoutTask = ThreadHelper.JoinableTaskFactory.RunAsync(UpdateCheckoutQueueAsync);
+					Logging.WriteLine($"Starting checkout task.");
+					StartTask = new TaskCompletionSource<bool>();
+					CheckoutTask = ThreadHelper.JoinableTaskFactory.RunAsync(async () => { await StartTask.Task; await UpdateCheckoutQueueAsync(); });
 				}
 			}
+			StartTask?.SetResult(true);
 		}
 
 		private async Task UpdateCheckoutQueueAsync()
@@ -750,6 +762,7 @@ namespace UnrealVS
 							QueueFile = CheckoutQueueFiles.FirstOrDefault();
 							if (QueueFile == null)
 							{
+								Logging.WriteLine($"Stopping checkout task; setting CheckoutTask to null.");
 								CheckoutTask = null;
 								return;
 							}
diff --git a/Engine/Source/Programs/UnrealVS/UnrealVS.Shared/UnrealVSPackage.cs b/Engine/Source/Programs/UnrealVS/UnrealVS.Shared/UnrealVSPackage.cs
--- a/Engine/Source/Programs/UnrealVS/UnrealVS.Shared/UnrealVSPackage.cs
+++ b/Engine/Source/Programs/UnrealVS/UnrealVS.Shared/UnrealVSPackage.cs
@@ -75,7 +75,7 @@ namespace UnrealVS
 	{
 		/** Constants */
 
-		private const string VersionString = "v1.77";
+		private const string VersionString = "v1.78";
 		private const string UnrealSolutionFileNamePrefix = "UE";
 		private const string ExtensionName = "UnrealVS";
 		private const string CommandLineOptionKey = ExtensionName + "CommandLineMRU";
