[UnrealVS] added Unreal.TestAdapter to run low level tests in the text explorer
added atomfeed.xml allow for auto updates. see README for details

diff --git a/Engine/Extras/UnrealVS/VS2022/AtomFeed.xml b/Engine/Extras/UnrealVS/VS2022/AtomFeed.xml
new file mode 100644
--- /dev/null
+++ b/Engine/Extras/UnrealVS/VS2022/AtomFeed.xml
@@ -0,0 +1,43 @@
+﻿<?xml version="1.0" encoding="utf-8"?>
+<feed xmlns="http://www.w3.org/2005/Atom">
+  <title type="text">UnrealVS Gallery</title>
+  <id>uuid:45b4571a-2981-49f1-9f70-04cfa96db5c8;id=1</id>
+  <updated>2022-10-21T18:08:01Z</updated>
+  <entry>
+    <id>ddbf523f-7eb6-4887-bd51-85a714ff87eb</id>
+    <title type="text">UnrealVS</title>
+    <summary type="text">Unreal Engine extension for Visual Studio</summary>
+    <published>2022-10-21T10:01:57-07:00</published>
+    <updated>2022-10-21T11:07:39-07:00</updated>
+    <author>
+      <name>Epic Games, Inc.</name>
+    </author>
+    <link rel="icon" href="ddbf523f-7eb6-4887-bd51-85a714ff87eb_Icon_Icon.png" />
+    <link rel="preview" href="ddbf523f-7eb6-4887-bd51-85a714ff87eb_Preview_Preview.png" />
+    <link rel="releasenotes" href="Resources\ReleaseNotes.txt" />
+    <link rel="alternate" href="https://docs.unrealengine.com/en-US/ProductionPipelines/DevelopmentSetup/VisualStudioSetup/UnrealVS/index.html" />
+    <content type="application/octet-stream" src="UnrealVS.vsix" />
+    <Vsix xmlns:xsd="http://www.w3.org/2001/XMLSchema" xmlns:xsi="http://www.w3.org/2001/XMLSchema-instance" xmlns="http://schemas.microsoft.com/developer/vsx-syndication-schema/2010">
+      <Id>ddbf523f-7eb6-4887-bd51-85a714ff87eb</Id>
+      <Version>1.79</Version>
+      <Rating xsi:nil="true" />
+      <RatingCount xsi:nil="true" />
+      <DownloadCount xsi:nil="true" />
+      <Installations>
+        <Identifier>Microsoft.VisualStudio.Pro</Identifier>
+        <VersionRange>[17.0,18.0)</VersionRange>
+        <ProductArchitecture>amd64</ProductArchitecture>
+      </Installations>
+      <Installations>
+        <Identifier>Microsoft.VisualStudio.Community</Identifier>
+        <VersionRange>[17.0,18.0)</VersionRange>
+        <ProductArchitecture>amd64</ProductArchitecture>
+      </Installations>
+      <Installations>
+        <Identifier>Microsoft.VisualStudio.Enterprise</Identifier>
+        <VersionRange>[17.0,18.0)</VersionRange>
+        <ProductArchitecture>amd64</ProductArchitecture>
+      </Installations>
+    </Vsix>
+  </entry>
+</feed>
\ No newline at end of file
diff --git a/Engine/Source/Programs/UnrealVS/README b/Engine/Source/Programs/UnrealVS/README
--- a/Engine/Source/Programs/UnrealVS/README
+++ b/Engine/Source/Programs/UnrealVS/README
@@ -1,3 +1,14 @@
 If nuget packages are updated, the minimum Visual Studio version needs to be updated to match in the .vsixmanifest.
 
-Ex. Microsoft.VisualStudio.SDK 16.9.31025.194 = Minimum Visual Studio version needs to be 16.9
\ No newline at end of file
+Ex. Microsoft.VisualStudio.SDK 16.9.31025.194 = Minimum Visual Studio version needs to be 16.9
+
+# Installation
+
+In Visual Studio
+Tool->Options->Environment->Extension
+Add an Addition Extension Gallery.
+Point the URL to the AtomFeed of where UnrealVS.vsix is synced
+ex.
+
+Name: UnrealVS
+URL: file:///D:/dev/ue5main/Engine/Extras/UnrealVS/VS2022/AtomFeed.xml
\ No newline at end of file
diff --git a/Engine/Source/Programs/UnrealVS/ReleaseNotes.txt b/Engine/Source/Programs/UnrealVS/ReleaseNotes.txt
--- a/Engine/Source/Programs/UnrealVS/ReleaseNotes.txt
+++ b/Engine/Source/Programs/UnrealVS/ReleaseNotes.txt
@@ -1,3 +1,6 @@
+1.79
+* added Unreal.TestAdapter to run low level tests from the text explorer
+
 1.78
 * fixed race condition when starting async checkout
 
diff --git a/Engine/Source/Programs/UnrealVS/Unreal.TestAdapter/Properties/AssemblyInfo.cs b/Engine/Source/Programs/UnrealVS/Unreal.TestAdapter/Properties/AssemblyInfo.cs
new file mode 100644
--- /dev/null
+++ b/Engine/Source/Programs/UnrealVS/Unreal.TestAdapter/Properties/AssemblyInfo.cs
@@ -0,0 +1,36 @@
+﻿using System.Reflection;
+using System.Runtime.CompilerServices;
+using System.Runtime.InteropServices;
+
+// General Information about an assembly is controlled through the following
+// set of attributes. Change these attribute values to modify the information
+// associated with an assembly.
+[assembly: AssemblyTitle("Unreal.TestAdapter")]
+[assembly: AssemblyDescription("")]
+[assembly: AssemblyConfiguration("")]
+[assembly: AssemblyCompany("Epic Games, Inc.")]
+[assembly: AssemblyProduct("Unreal.TestAdapter")]
+[assembly: AssemblyCopyright("Copyright Epic Games, Inc.")]
+[assembly: AssemblyTrademark("")]
+[assembly: AssemblyCulture("")]
+
+// Setting ComVisible to false makes the types in this assembly not visible
+// to COM components.  If you need to access a type in this assembly from
+// COM, set the ComVisible attribute to true on that type.
+[assembly: ComVisible(false)]
+
+// The following GUID is for the ID of the typelib if this project is exposed to COM
+[assembly: Guid("2571af3d-afee-4a45-b612-2795d5a5fe6b")]
+
+// Version information for an assembly consists of the following four values:
+//
+//      Major Version
+//      Minor Version
+//      Build Number
+//      Revision
+//
+// You can specify all the values or you can default the Build and Revision Numbers
+// by using the '*' as shown below:
+// [assembly: AssemblyVersion("1.0.*")]
+[assembly: AssemblyVersion("1.0.0.0")]
+[assembly: AssemblyFileVersion("1.0.0.0")]
diff --git a/Engine/Source/Programs/UnrealVS/Unreal.TestAdapter/ReadMe.md b/Engine/Source/Programs/UnrealVS/Unreal.TestAdapter/ReadMe.md
new file mode 100644
--- /dev/null
+++ b/Engine/Source/Programs/UnrealVS/Unreal.TestAdapter/ReadMe.md
@@ -0,0 +1,16 @@
+#Things to know
+
+* dll must be called XXX.TestAdapter.dll. VS looks for dll's with this pattern.
+
+* source.extension.vsixmanifest has a reference to this project. that does the deployment of the dll. 
+   Under Assets with a type of `UnitTestExtension`
+
+* TestDiscoverer relies on a file called `(exe).is_unreal_test existing`.
+   The TestTargetRules currently outputs this file
+   The reason for this is to quickly identify text executables. 
+   Another option would be add a special string into the exe and then look for it in the binary
+
+#Debugging
+
+Easiest way to debug the TestDisoverer or the TestExecutor by adding `System.Diagnostics.Debugger.Launch();` to 
+`DiscoverTests` or `RunTests` respectively. 
diff --git a/Engine/Source/Programs/UnrealVS/Unreal.TestAdapter/TestDiscoverer.cs b/Engine/Source/Programs/UnrealVS/Unreal.TestAdapter/TestDiscoverer.cs
new file mode 100644
--- /dev/null
+++ b/Engine/Source/Programs/UnrealVS/Unreal.TestAdapter/TestDiscoverer.cs
@@ -0,0 +1,153 @@
+using System;
+using System.Collections.Generic;
+using System.Diagnostics;
+using System.IO;
+using System.Text;
+using System.Xml;
+using Microsoft.VisualStudio.TestPlatform.ObjectModel;
+using Microsoft.VisualStudio.TestPlatform.ObjectModel.Adapter;
+using Microsoft.VisualStudio.TestPlatform.ObjectModel.Logging;
+
+namespace Unreal.TestAdapter
+{
+	[FileExtension(".exe")]
+	[DefaultExecutorUri("executor://UnrealTestExecutor")]
+	class TestDiscoverer : ITestDiscoverer
+	{
+		static readonly Uri ExecutorUri = new Uri("executor://UnrealTestExecutor");
+		internal static readonly TestProperty DisabledTestProperty = TestProperty.Register("TestDiscoverer.DisabledTestProperty", "DisabledTestProperty", string.Empty, string.Empty, typeof(bool), (object o) => o is bool, TestPropertyAttributes.Hidden, typeof(TestDiscoverer));
+
+		public void DiscoverTests(IEnumerable<string> sources, IDiscoveryContext discoveryContext, IMessageLogger logger, ITestCaseDiscoverySink discoverySink)
+		{
+			foreach (var source in sources)
+			{
+				//cheap way to look for tests. outputted by the TestTargetRules
+				if (!File.Exists(source + ".is_unreal_test"))
+					continue;
+
+				string xml = GetTestCaseXml(source, logger);
+				ParseXml(xml, discoverySink, source);
+			}
+		}
+
+		string GetTestCaseXml(string source, IMessageLogger logger)
+		{
+			Process process = new Process();
+			process.StartInfo.FileName = source;
+			process.StartInfo.UseShellExecute = false;
+			process.StartInfo.RedirectStandardOutput = true;
+			process.StartInfo.CreateNoWindow = true;
+			process.StartInfo.Arguments = "--list-tests * --reporter xml";
+			process.StartInfo.RedirectStandardError = true;
+
+			StringBuilder error = new StringBuilder();
+			process.ErrorDataReceived += new DataReceivedEventHandler((sender, e) =>
+			{
+				error.Append(e.Data);
+			});
+
+			process.Start();
+			string result = process.StandardOutput.ReadToEnd();
+			bool exited = process.WaitForExit(500);
+			int exitCode = process.ExitCode;
+			if (!exited || exitCode != 0)
+			{
+				logger.SendMessage(TestMessageLevel.Error, "Error " + error);
+			}
+			return result;
+
+		}
+
+		static void ParseXml(string xml, ITestCaseDiscoverySink discoverySink, string source)
+		{
+			TextReader xmlText = new StringReader(xml);
+			using (var xmlReader = XmlReader.Create(xmlText, new XmlReaderSettings() { IgnoreWhitespace = true }))
+			{
+				while (xmlReader.Read())
+				{
+					switch (xmlReader.NodeType)
+					{
+						case XmlNodeType.Element:
+							if (xmlReader.Name == "TestCase" && !xmlReader.IsEmptyElement)
+							{
+								if (ParseTestCase(xmlReader, source, out TestCase testCase))
+								{
+									discoverySink.SendTestCase(testCase);
+								}
+							}
+							break;
+						case XmlNodeType.EndElement:
+							if (xmlReader.Name == "MatchingTests")
+							{
+								//quit because sometimes there is extra lines are the xml
+								//ie. Shutting down and abandoning module CoreUObject (4)
+								return;
+							}
+							break;
+						default:
+							break;
+					}
+				}
+			}
+		}
+		static bool ParseTestCase(XmlReader xmlReader, string source, out TestCase testCase)
+		{
+			testCase = null;
+			//advance past the start element
+			xmlReader.Read();
+
+			string fullyQualifiedName = null;
+			string codeFilePath = null;
+			int lineNumber = 0;
+			bool disabled = false;
+			while (xmlReader.NodeType != XmlNodeType.EndElement)
+			{
+				if (xmlReader.IsEmptyElement)
+				{
+					xmlReader.Skip();
+				}
+				else if (xmlReader.Name == "Name")
+				{
+					fullyQualifiedName = xmlReader.ReadElementContentAsString();
+				}
+				else if (xmlReader.Name == "SourceInfo")
+				{
+					xmlReader.Read();
+					while (xmlReader.NodeType != XmlNodeType.EndElement)
+					{
+						if (xmlReader.Name == "File")
+						{
+							codeFilePath = xmlReader.ReadElementContentAsString();
+						}
+						else if (xmlReader.Name == "Line")
+						{
+							lineNumber = xmlReader.ReadElementContentAsInt();
+						}
+						else
+						{
+							xmlReader.Read();
+						}
+					}
+					xmlReader.Read(); //read </SourceInfo>
+				}
+				else if(xmlReader.Name == "Tags")
+				{
+					string tags = xmlReader.ReadElementContentAsString();
+					disabled = tags.Contains("[.]");
+				}
+				else
+				{
+					xmlReader.Skip(); //skip elements we don't care about
+				}
+			}
+			if (codeFilePath != null && fullyQualifiedName != null)
+			{
+				testCase = new TestCase(fullyQualifiedName, ExecutorUri, source) { LineNumber = lineNumber, CodeFilePath = codeFilePath };
+				testCase.SetPropertyValue(DisabledTestProperty, disabled);
+			}
+			return testCase != null;
+		}
+	}
+
+	
+}
diff --git a/Engine/Source/Programs/UnrealVS/Unreal.TestAdapter/TestExecutor.cs b/Engine/Source/Programs/UnrealVS/Unreal.TestAdapter/TestExecutor.cs
new file mode 100644
--- /dev/null
+++ b/Engine/Source/Programs/UnrealVS/Unreal.TestAdapter/TestExecutor.cs
@@ -0,0 +1,163 @@
+using System;
+using System.Collections.Generic;
+using System.Diagnostics;
+using System.Diagnostics.Eventing.Reader;
+using System.IO;
+using System.Threading;
+using System.Threading.Tasks;
+using System.Xml;
+using Microsoft.VisualStudio.TestPlatform.ObjectModel;
+using Microsoft.VisualStudio.TestPlatform.ObjectModel.Adapter;
+
+namespace Unreal.TestAdapter
+{
+	[ExtensionUri("executor://UnrealTestExecutor")]
+	class TestExecutor : ITestExecutor
+	{
+		CancellationTokenSource CancellationTokenSource;
+		bool WasCancelled = false;
+
+		//this does appear ever be called.
+		//meaning that launched process will continue running
+		public void Cancel()
+		{
+			WasCancelled = true;
+			CancellationTokenSource?.Cancel();
+		}
+
+		public void RunTests(IEnumerable<TestCase> tests, IRunContext runContext, IFrameworkHandle frameworkHandle)
+		{
+			//early out if canceled before even started
+			if (WasCancelled)
+				return;
+
+			using (CancellationTokenSource = new CancellationTokenSource())
+			{
+				var token = CancellationTokenSource.Token;
+				Parallel.ForEach(tests, testCase =>
+					{
+						//if canceled skip but continue the loop to let the other processes to shutdown
+						if (!token.IsCancellationRequested)
+						{
+							RunTestCase(testCase, frameworkHandle, token, runContext.IsBeingDebugged);
+						}
+					});
+			}
+		}
+
+		void RunTestCase(TestCase testCase, IFrameworkHandle frameworkHandle, CancellationToken token, bool debug)
+		{
+			if(testCase.GetPropertyValue(TestDiscoverer.DisabledTestProperty, false))
+			{
+				frameworkHandle.RecordResult(new TestResult(testCase) { Outcome = TestOutcome.Skipped });
+				return;
+			}
+			string xml = null;
+
+			if (debug)
+			{
+				string tempFile = Path.Combine(Path.GetTempPath(), Path.GetTempFileName());
+				int processId = frameworkHandle.LaunchProcessWithDebuggerAttached(testCase.Source, Path.GetDirectoryName(testCase.Source), "\"" + testCase.FullyQualifiedName + "\"  --reporter unreal -o " + tempFile, null);
+				Process process = Process.GetProcessById(processId);
+				process.WaitForExit();
+				if (File.Exists(tempFile))
+				{
+					xml = File.ReadAllText(tempFile);
+				}
+			}
+			else
+			{
+				using (Process process = new Process())
+				{
+					process.StartInfo.FileName = testCase.Source;
+
+					//wrap in quotes because test names can have spaces
+					process.StartInfo.Arguments = "\"" + testCase.FullyQualifiedName + "\"  --reporter unreal";
+					process.StartInfo.RedirectStandardOutput = true;
+					process.StartInfo.UseShellExecute = false;
+
+					if (token.IsCancellationRequested)
+					{
+						return;
+					}
+
+					frameworkHandle.RecordStart(testCase);
+					process.Start();
+					xml = process.StandardOutput.ReadToEnd();
+					while (!process.HasExited)
+					{
+						if (token.IsCancellationRequested)
+						{
+							process.Kill();
+							return;
+						}
+					}
+				}
+			}
+
+			var testResult = new TestResult(testCase) { Outcome = TestOutcome.NotFound };
+			ParseResult(xml, testResult);
+			frameworkHandle.RecordEnd(testCase, testResult.Outcome);
+			frameworkHandle.RecordResult(testResult);
+		}
+
+		void ParseResult(string xmlString, TestResult testResult)
+		{
+			TextReader textReader = new StringReader(xmlString);
+			using (XmlReader xml = XmlReader.Create(textReader, new XmlReaderSettings() { IgnoreWhitespace = true }))
+			{
+				while (xml.Read())
+				{
+					switch (xml.NodeType)
+					{
+						case XmlNodeType.Element:
+							if (xml.Name == "testcase" && xml.GetAttribute("name") == testResult.TestCase.FullyQualifiedName)
+							{
+								ParseTestCaseResult(xml, testResult);
+							}
+							break;
+						case XmlNodeType.EndElement:
+							if (xml.Name == "testrun")
+							{
+								return; //quit to avoid junk at the end of the output
+							}
+							break;
+						default:
+							break;
+					}
+				}
+			}
+
+		}
+
+		void ParseTestCaseResult(XmlReader xml, TestResult testResult)
+		{
+			xml.Read();
+			while (xml.NodeType != XmlNodeType.EndElement)
+			{
+				if (xml.Name == "failure")
+				{
+					string failure = xml.ReadElementContentAsString()?.Trim() ?? String.Empty;
+					testResult.ErrorStackTrace = failure;
+					testResult.Outcome = TestOutcome.Failed;
+				}
+				else if (xml.Name == "result")
+				{
+					string strSuccess = xml.GetAttribute("success");
+					bool.TryParse(strSuccess, out bool success);
+					testResult.Outcome = success ? TestOutcome.Passed : TestOutcome.Failed;
+					xml.Read();
+				}
+				else
+				{
+					xml.Skip();
+				}
+			}
+		}
+
+		public void RunTests(IEnumerable<string> sources, IRunContext runContext, IFrameworkHandle frameworkHandle)
+		{
+			
+		}
+	}
+}
diff --git a/Engine/Source/Programs/UnrealVS/Unreal.TestAdapter/Unreal.TestAdapter.csproj b/Engine/Source/Programs/UnrealVS/Unreal.TestAdapter/Unreal.TestAdapter.csproj
new file mode 100644
--- /dev/null
+++ b/Engine/Source/Programs/UnrealVS/Unreal.TestAdapter/Unreal.TestAdapter.csproj
@@ -0,0 +1,76 @@
+﻿<?xml version="1.0" encoding="utf-8"?>
+<Project ToolsVersion="15.0" xmlns="http://schemas.microsoft.com/developer/msbuild/2003">
+  <Import Project="$(MSBuildExtensionsPath)\$(MSBuildToolsVersion)\Microsoft.Common.props" Condition="Exists('$(MSBuildExtensionsPath)\$(MSBuildToolsVersion)\Microsoft.Common.props')" />
+  <PropertyGroup>
+    <Configuration Condition=" '$(Configuration)' == '' ">Debug</Configuration>
+    <Platform Condition=" '$(Platform)' == '' ">AnyCPU</Platform>
+    <ProjectGuid>{2571AF3D-AFEE-4A45-B612-2795D5A5FE6B}</ProjectGuid>
+    <OutputType>Library</OutputType>
+    <AppDesignerFolder>Properties</AppDesignerFolder>
+    <RootNamespace>Unreal.TestAdapter</RootNamespace>
+    <AssemblyName>Unreal.TestAdapter</AssemblyName>
+    <TargetFrameworkVersion>v4.7.2</TargetFrameworkVersion>
+    <FileAlignment>512</FileAlignment>
+    <Deterministic>true</Deterministic>
+  </PropertyGroup>
+  <PropertyGroup Condition=" '$(Configuration)|$(Platform)' == 'Debug|AnyCPU' ">
+    <DebugSymbols>true</DebugSymbols>
+    <DebugType>full</DebugType>
+    <Optimize>false</Optimize>
+    <OutputPath>bin\Debug\</OutputPath>
+    <DefineConstants>DEBUG;TRACE</DefineConstants>
+    <ErrorReport>prompt</ErrorReport>
+    <WarningLevel>4</WarningLevel>
+  </PropertyGroup>
+  <PropertyGroup Condition=" '$(Configuration)|$(Platform)' == 'Release|AnyCPU' ">
+    <DebugType>pdbonly</DebugType>
+    <Optimize>true</Optimize>
+    <OutputPath>bin\Release\</OutputPath>
+    <DefineConstants>TRACE</DefineConstants>
+    <ErrorReport>prompt</ErrorReport>
+    <WarningLevel>4</WarningLevel>
+  </PropertyGroup>
+  <ItemGroup>
+    <Reference Include="Microsoft.TestPlatform.CoreUtilities, Version=15.0.0.0, Culture=neutral, PublicKeyToken=b03f5f7f11d50a3a, processorArchitecture=MSIL">
+      <HintPath>..\packages\Microsoft.TestPlatform.ObjectModel.16.10.0\lib\net451\Microsoft.TestPlatform.CoreUtilities.dll</HintPath>
+    </Reference>
+    <Reference Include="Microsoft.TestPlatform.PlatformAbstractions, Version=15.0.0.0, Culture=neutral, PublicKeyToken=b03f5f7f11d50a3a, processorArchitecture=MSIL">
+      <HintPath>..\packages\Microsoft.TestPlatform.ObjectModel.16.10.0\lib\net451\Microsoft.TestPlatform.PlatformAbstractions.dll</HintPath>
+    </Reference>
+    <Reference Include="Microsoft.VisualStudio.TestPlatform.ObjectModel, Version=15.0.0.0, Culture=neutral, PublicKeyToken=b03f5f7f11d50a3a, processorArchitecture=MSIL">
+      <HintPath>..\packages\Microsoft.TestPlatform.ObjectModel.16.10.0\lib\net451\Microsoft.VisualStudio.TestPlatform.ObjectModel.dll</HintPath>
+    </Reference>
+    <Reference Include="NuGet.Frameworks, Version=5.11.0.10, Culture=neutral, PublicKeyToken=31bf3856ad364e35, processorArchitecture=MSIL">
+      <HintPath>..\packages\NuGet.Frameworks.5.11.0\lib\net472\NuGet.Frameworks.dll</HintPath>
+    </Reference>
+    <Reference Include="System" />
+    <Reference Include="System.Collections.Immutable, Version=1.2.3.0, Culture=neutral, PublicKeyToken=b03f5f7f11d50a3a, processorArchitecture=MSIL">
+      <HintPath>..\packages\System.Collections.Immutable.1.5.0\lib\netstandard2.0\System.Collections.Immutable.dll</HintPath>
+    </Reference>
+    <Reference Include="System.Configuration" />
+    <Reference Include="System.Core" />
+    <Reference Include="System.Reflection.Metadata, Version=1.4.3.0, Culture=neutral, PublicKeyToken=b03f5f7f11d50a3a, processorArchitecture=MSIL">
+      <HintPath>..\packages\System.Reflection.Metadata.1.6.0\lib\netstandard2.0\System.Reflection.Metadata.dll</HintPath>
+    </Reference>
+    <Reference Include="System.Runtime" />
+    <Reference Include="System.Runtime.Serialization" />
+    <Reference Include="System.Xml.Linq" />
+    <Reference Include="System.Data.DataSetExtensions" />
+    <Reference Include="Microsoft.CSharp" />
+    <Reference Include="System.Data" />
+    <Reference Include="System.Net.Http" />
+    <Reference Include="System.Xml" />
+  </ItemGroup>
+  <ItemGroup>
+    <Compile Include="TestDiscoverer.cs" />
+    <Compile Include="Properties\AssemblyInfo.cs" />
+    <Compile Include="TestExecutor.cs" />
+  </ItemGroup>
+  <ItemGroup>
+    <None Include="packages.config" />
+  </ItemGroup>
+  <ItemGroup>
+    <None Include="ReadMe.md" />
+  </ItemGroup>
+  <Import Project="$(MSBuildToolsPath)\Microsoft.CSharp.targets" />
+</Project>
\ No newline at end of file
diff --git a/Engine/Source/Programs/UnrealVS/UnrealVS.2019/Properties/AssemblyInfo.cs b/Engine/Source/Programs/UnrealVS/UnrealVS.2019/Properties/AssemblyInfo.cs
--- a/Engine/Source/Programs/UnrealVS/UnrealVS.2019/Properties/AssemblyInfo.cs
+++ b/Engine/Source/Programs/UnrealVS/UnrealVS.2019/Properties/AssemblyInfo.cs
@@ -31,8 +31,8 @@ using System.Runtime.InteropServices;
 // You can specify all the values or you can default the Revision and Build Numbers 
 // by using the '*' as shown below:
 
-[assembly: AssemblyVersion("1.78.0.0")]
-[assembly: AssemblyFileVersion("1.78.0.0")]
+[assembly: AssemblyVersion("1.79.0.0")]
+[assembly: AssemblyFileVersion("1.79.0.0")]
 
 
 
diff --git a/Engine/Source/Programs/UnrealVS/UnrealVS.2022/UnrealVS.2022.csproj b/Engine/Source/Programs/UnrealVS/UnrealVS.2022/UnrealVS.2022.csproj
--- a/Engine/Source/Programs/UnrealVS/UnrealVS.2022/UnrealVS.2022.csproj
+++ b/Engine/Source/Programs/UnrealVS/UnrealVS.2022/UnrealVS.2022.csproj
@@ -101,6 +101,11 @@
       <CopyToOutputDirectory>Always</CopyToOutputDirectory>
       <IncludeInVSIX>true</IncludeInVSIX>
     </Content>
+    <Content Include="..\ReleaseNotes.txt">
+      <Link>Resources\ReleaseNotes.txt</Link>
+      <CopyToOutputDirectory>Always</CopyToOutputDirectory>
+      <IncludeInVSIX>true</IncludeInVSIX>
+    </Content>
     <Content Include="Resources\Icon.png">
       <CopyToOutputDirectory>Always</CopyToOutputDirectory>
       <IncludeInVSIX>true</IncludeInVSIX>
@@ -140,6 +145,9 @@
     <PackageReference Include="Microsoft.CSharp">
       <Version>4.7.0</Version>
     </PackageReference>
+    <PackageReference Include="Microsoft.TestPlatform.ObjectModel">
+      <Version>17.3.2</Version>
+    </PackageReference>
     <PackageReference Include="Microsoft.VisualStudio.Sdk" Version="17.0.0-previews-1-31410-273">
     </PackageReference>
     <PackageReference Include="Microsoft.VisualStudio.Threading">
@@ -157,10 +165,20 @@
       <IncludeAssets>runtime; build; native; contentfiles; analyzers; buildtransitive</IncludeAssets>
       <PrivateAssets>all</PrivateAssets>
     </PackageReference>
+    <PackageReference Include="MSTest.TestAdapter">
+      <Version>2.2.10</Version>
+    </PackageReference>
     <PackageReference Include="System.Text.Json">
       <Version>6.0.0-preview.5.21301.5</Version>
     </PackageReference>
   </ItemGroup>
+  <ItemGroup>
+    <ProjectReference Include="..\Unreal.TestAdapter\Unreal.TestAdapter.csproj">
+      <Project>{2571AF3D-AFEE-4A45-B612-2795D5A5FE6B}</Project>
+      <Name>Unreal.TestAdapter</Name>
+      <IncludeOutputGroupsInVSIX>BuiltProjectOutputGroup%3bBuiltProjectOutputGroupDependencies%3bGetCopyToOutputDirectoryItems%3bSatelliteDllsProjectOutputGroup%3b</IncludeOutputGroupsInVSIX>
+    </ProjectReference>
+  </ItemGroup>
   <Import Project="..\UnrealVS.Shared\UnrealVS.Shared.projitems" Label="Shared" />
   <PropertyGroup>
     <UseCodebase>true</UseCodebase>
diff --git a/Engine/Source/Programs/UnrealVS/UnrealVS.Shared/Utils.cs b/Engine/Source/Programs/UnrealVS/UnrealVS.Shared/Utils.cs
--- a/Engine/Source/Programs/UnrealVS/UnrealVS.Shared/Utils.cs
+++ b/Engine/Source/Programs/UnrealVS/UnrealVS.Shared/Utils.cs
@@ -12,6 +12,7 @@ using System.IO;
 using System.Linq;
 using System.Windows.Controls;
 using System.Windows.Media;
+using System.Windows.Media.Animation;
 using System.Xml;
 
 namespace UnrealVS
@@ -536,7 +537,14 @@ namespace UnrealVS
 					{
 						foreach (DirectoryInfo ProjectDir in ParentProjectDir.EnumerateDirectories())
 						{
-							ProjectFiles.AddRange(ProjectDir.EnumerateFiles("*.uproject"));
+							try
+							{
+								ProjectFiles.AddRange(ProjectDir.EnumerateFiles("*.uproject"));
+							} 
+							catch(Exception Ex)
+							{
+								Logging.WriteLine($"AddingProjects: Exception trying to add projects from directory '{ProjectDir}': {Ex.Message}");
+							}
 						}
 					}
 				}
