[UnrealVS] added missing copyrights
updated version number and instructions

diff --git a/Engine/Source/Programs/UnrealVS/Unreal.TestAdapter/TestDiscoverer.cs b/Engine/Source/Programs/UnrealVS/Unreal.TestAdapter/TestDiscoverer.cs
--- a/Engine/Source/Programs/UnrealVS/Unreal.TestAdapter/TestDiscoverer.cs
+++ b/Engine/Source/Programs/UnrealVS/Unreal.TestAdapter/TestDiscoverer.cs
@@ -1,3 +1,5 @@
+// Copyright Epic Games, Inc. All Rights Reserved.
+
 using System;
 using System.Collections.Generic;
 using System.Diagnostics;
diff --git a/Engine/Source/Programs/UnrealVS/Unreal.TestAdapter/TestExecutor.cs b/Engine/Source/Programs/UnrealVS/Unreal.TestAdapter/TestExecutor.cs
--- a/Engine/Source/Programs/UnrealVS/Unreal.TestAdapter/TestExecutor.cs
+++ b/Engine/Source/Programs/UnrealVS/Unreal.TestAdapter/TestExecutor.cs
@@ -1,7 +1,8 @@
+// Copyright Epic Games, Inc. All Rights Reserved.
+
 using System;
 using System.Collections.Generic;
 using System.Diagnostics;
-using System.Diagnostics.Eventing.Reader;
 using System.IO;
 using System.Threading;
 using System.Threading.Tasks;
diff --git a/Engine/Source/Programs/UnrealVS/UnrealVS.2022/Properties/AssemblyInfo.cs b/Engine/Source/Programs/UnrealVS/UnrealVS.2022/Properties/AssemblyInfo.cs
--- a/Engine/Source/Programs/UnrealVS/UnrealVS.2022/Properties/AssemblyInfo.cs
+++ b/Engine/Source/Programs/UnrealVS/UnrealVS.2022/Properties/AssemblyInfo.cs
@@ -30,8 +30,8 @@ using System.Runtime.InteropServices;
 // You can specify all the values or you can default the Revision and Build Numbers 
 // by using the '*' as shown below:
 
-[assembly: AssemblyVersion("1.78.0.0")]
-[assembly: AssemblyFileVersion("1.78.0.0")]
+[assembly: AssemblyVersion("1.79.0.0")]
+[assembly: AssemblyFileVersion("1.79.0.0")]
 
 
 
diff --git a/Engine/Source/Programs/UnrealVS/UnrealVS.Shared/UnrealVSPackage.cs b/Engine/Source/Programs/UnrealVS/UnrealVS.Shared/UnrealVSPackage.cs
--- a/Engine/Source/Programs/UnrealVS/UnrealVS.Shared/UnrealVSPackage.cs
+++ b/Engine/Source/Programs/UnrealVS/UnrealVS.Shared/UnrealVSPackage.cs
@@ -75,7 +75,7 @@ namespace UnrealVS
 	{
 		/** Constants */
 
-		private const string VersionString = "v1.78";
+		private const string VersionString = "v1.79";
 		private const string UnrealSolutionFileNamePrefix = "UE";
 		private const string ExtensionName = "UnrealVS";
 		private const string CommandLineOptionKey = ExtensionName + "CommandLineMRU";
diff --git a/Engine/Source/Programs/UnrealVS/UpdateInstructions.txt b/Engine/Source/Programs/UnrealVS/UpdateInstructions.txt
--- a/Engine/Source/Programs/UnrealVS/UpdateInstructions.txt
+++ b/Engine/Source/Programs/UnrealVS/UpdateInstructions.txt
@@ -1,5 +1,8 @@
 1. Make your changes
 2. Update version number in both VSIX manifest files and both AssemblyInfo files.
 3. update ReleaseNotes.txt with the changes
-4. Checkout the VSIX files and rebuild 2019 IN VS2019, 2022 in VS2022
-5. Checkin all the above in 1 CL
+4. Checkout the VSIX in Engine\Extras\UnrealVS\ files and rebuild 2019 IN VS2019, 2022 in VS2022
+5. Update the atomfeed.xml in Engine\Extras\UnrealVS\<version>\atomfeed.xml
+	* <updated> with new timespamp
+	* <Version> with new version number
+6. Checkin all the above in 1 CL
