[UnrealVS] added missing copyrights

diff --git a/Engine/Source/Programs/UnrealVS/Unreal.TestAdapter/Properties/AssemblyInfo.cs b/Engine/Source/Programs/UnrealVS/Unreal.TestAdapter/Properties/AssemblyInfo.cs
--- a/Engine/Source/Programs/UnrealVS/Unreal.TestAdapter/Properties/AssemblyInfo.cs
+++ b/Engine/Source/Programs/UnrealVS/Unreal.TestAdapter/Properties/AssemblyInfo.cs
@@ -1,5 +1,6 @@
-﻿using System.Reflection;
-using System.Runtime.CompilerServices;
+﻿// Copyright Epic Games, Inc. All Rights Reserved.
+
+using System.Reflection;
 using System.Runtime.InteropServices;
 
 // General Information about an assembly is controlled through the following
