[UnrealVS] changed TestAdapater to report durations
refactored projects to share files avoiding having to duplicate changes between VS2019 and VS022
added menu to improved discoverability of features
added option to use P4V diff instead of VS diff

diff --git a/Engine/Extras/UnrealVS/VS2022/AtomFeed.xml b/Engine/Extras/UnrealVS/VS2022/AtomFeed.xml
--- a/Engine/Extras/UnrealVS/VS2022/AtomFeed.xml
+++ b/Engine/Extras/UnrealVS/VS2022/AtomFeed.xml
@@ -8,7 +8,7 @@
     <title type="text">UnrealVS</title>
     <summary type="text">Unreal Engine extension for Visual Studio</summary>
     <published>2022-10-21T10:01:57-07:00</published>
-    <updated>2022-10-21T11:07:39-07:00</updated>
+    <updated>2022-10-28T11:09:39-07:00</updated>
     <author>
       <name>Epic Games, Inc.</name>
     </author>
@@ -19,7 +19,7 @@
     <content type="application/octet-stream" src="UnrealVS.vsix" />
     <Vsix xmlns:xsd="http://www.w3.org/2001/XMLSchema" xmlns:xsi="http://www.w3.org/2001/XMLSchema-instance" xmlns="http://schemas.microsoft.com/developer/vsx-syndication-schema/2010">
       <Id>ddbf523f-7eb6-4887-bd51-85a714ff87eb</Id>
-      <Version>1.79</Version>
+      <Version>1.80</Version>
       <Rating xsi:nil="true" />
       <RatingCount xsi:nil="true" />
       <DownloadCount xsi:nil="true" />
diff --git a/Engine/Source/Developer/LowLevelTestsRunner/Private/UnrealReporter.cpp b/Engine/Source/Developer/LowLevelTestsRunner/Private/UnrealReporter.cpp
--- a/Engine/Source/Developer/LowLevelTestsRunner/Private/UnrealReporter.cpp
+++ b/Engine/Source/Developer/LowLevelTestsRunner/Private/UnrealReporter.cpp
@@ -2,13 +2,11 @@
 
 #include <catch2/reporters/catch_reporter_streaming_base.hpp>
 #include <catch2/catch_test_case_info.hpp>
+#include <catch2/catch_timer.hpp>
 #include <catch2/reporters/catch_reporter_registrars.hpp>
 #include <catch2/internal/catch_xmlwriter.hpp>
-#include <catch2/internal/catch_textflow.hpp>
 #include "Misc/StringBuilder.h"
 
-#include <iostream>
-
 namespace UE::LowLevelTests
 {
 /** Simple reporter for the low level tests
@@ -18,7 +16,7 @@ namespace UE::LowLevelTests
 		<failure>
 			all error text 
 		</failure>
-		<result success="true/false" />
+		<result success="true/false" duration="double seconds"/>
 	</testcase>
 </testrun>
 
@@ -50,6 +48,7 @@ public:
 
 	void testCaseStarting(Catch::TestCaseInfo const& TestCaseInfo) override
 	{
+		Timer.start();
 		Xml.startElement("testcase");
 		Xml.writeAttribute("name", TestCaseInfo.name);
 	}
@@ -57,6 +56,8 @@ public:
 	virtual void testCaseEnded(Catch::TestCaseStats const& TestCaseStats) override
 	{
 		using namespace Catch;
+		double Duration = Timer.getElapsedSeconds();
+
 		if (bHasFailedOnce)
 		{
 			Xml.endElement(); //end failure element
@@ -67,6 +68,7 @@ public:
 			//otherwise any failures would have to be stored as a member
 			XmlWriter::ScopedElement ResultElement = Xml.scopedElement("result");
 			ResultElement.writeAttribute("success", TestCaseStats.totals.assertions.allOk());
+			ResultElement.writeAttribute("duration", Duration);
 		}
 		Xml.endElement();
 	}
@@ -153,6 +155,7 @@ public:
 
 private:
 	Catch::XmlWriter Xml;
+	Catch::Timer Timer;
 	bool bHasFailedOnce = false;
 };
 
diff --git a/Engine/Source/Programs/LowLevelTests/README.md b/Engine/Source/Programs/LowLevelTests/README.md
--- a/Engine/Source/Programs/LowLevelTests/README.md
+++ b/Engine/Source/Programs/LowLevelTests/README.md
@@ -197,15 +197,12 @@ TEST_CASE("Summary of test case", "[unit][feature][fast]")
 ### From Visual Studio
 
 Tests can be built and run directly from Visual Studio on desktop platforms:
-1. Download the [Test Adapter for Catch2](https://marketplace.visualstudio.com/items?itemName=JohnnyHendriks.ext01) and install it.
-2. Build the test projects from Visual Studio to produce the executables. The test adapter will discover tests in the executables and they must contain the word **Test** e.g. *MyModuleTest.exe* or *MyModuleTests.exe*.
+1. Install UnrealVS
+2. Build the test projects from Visual Studio to produce the executables. The test adapter will discover tests in the executables.
 3. The tests will be displayed in the test explorer: select *Test -> Test Explorer* from the menu: from here you can run tests and navigate to their source code.
 Running the tests will generate a *Test Detail Summary* and a Catch test report XML file in the same folder as the executable.
 4. If there are no tests in Test Explorer it's likely because the `Build` at step 3. didn't update the executable. Run a `Rebuild` on the test project and you should be good to go.
 
-> **If there are still no tests shown in the *Test Explorer*, perform a Rebuild on the test project and in the Output tab in VS select `Show output from: Tests` and look for `Catch2Adapter` in the log.
-> You should see a `Started Catch2Adapter test discovery...` line and a `Finished Catch2Adapter test discovery.`: between these lines you will see any problems reported by the Catch2 test adapter.**
-
 ## From command line
 
 Go to the binaries folder and run the executable with the extra flag `-- --wait` to wait for user input before exiting. By default the test app exits as soon as it's done, this is intentional to prevent easy hangs on automated systems and tools such as the `Test Adapter for Catch2` which runs the executable to list all the tests.
diff --git a/Engine/Source/Programs/UnrealVS/ReleaseNotes.txt b/Engine/Source/Programs/UnrealVS/ReleaseNotes.txt
--- a/Engine/Source/Programs/UnrealVS/ReleaseNotes.txt
+++ b/Engine/Source/Programs/UnrealVS/ReleaseNotes.txt
@@ -1,3 +1,10 @@
+1.80
+* updated Unreal.TestAdapater to report duration's
+* added a menu for extension to improve discover-ability of features
+* added option to use P4V diff instead of Visual Studio diff
+* refactored code to share files between VS2019 and VS2022 versions
+* chanaged the File Browser to handle paste of file names
+
 1.79
 * added Unreal.TestAdapter to run low level tests from the text explorer
 
diff --git a/Engine/Source/Programs/UnrealVS/Unreal.TestAdapter/Properties/AssemblyInfo.cs b/Engine/Source/Programs/UnrealVS/Unreal.TestAdapter/Properties/AssemblyInfo.cs
--- a/Engine/Source/Programs/UnrealVS/Unreal.TestAdapter/Properties/AssemblyInfo.cs
+++ b/Engine/Source/Programs/UnrealVS/Unreal.TestAdapter/Properties/AssemblyInfo.cs
@@ -33,5 +33,5 @@ using System.Runtime.InteropServices;
 // You can specify all the values or you can default the Build and Revision Numbers
 // by using the '*' as shown below:
 // [assembly: AssemblyVersion("1.0.*")]
-[assembly: AssemblyVersion("1.0.0.0")]
-[assembly: AssemblyFileVersion("1.0.0.0")]
+[assembly: AssemblyVersion("1.1.0.0")]
+[assembly: AssemblyFileVersion("1.1.0.0")]
diff --git a/Engine/Source/Programs/UnrealVS/Unreal.TestAdapter/TestExecutor.cs b/Engine/Source/Programs/UnrealVS/Unreal.TestAdapter/TestExecutor.cs
--- a/Engine/Source/Programs/UnrealVS/Unreal.TestAdapter/TestExecutor.cs
+++ b/Engine/Source/Programs/UnrealVS/Unreal.TestAdapter/TestExecutor.cs
@@ -147,6 +147,14 @@ namespace Unreal.TestAdapter
 					string strSuccess = xml.GetAttribute("success");
 					bool.TryParse(strSuccess, out bool success);
 					testResult.Outcome = success ? TestOutcome.Passed : TestOutcome.Failed;
+
+					string strDuration = xml.GetAttribute("duration");
+					if (strDuration != null)
+					{
+						double.TryParse(strDuration, out double duraction);
+						testResult.Duration = TimeSpan.FromSeconds(duraction);
+					}
+
 					xml.Read();
 				}
 				else
diff --git a/Engine/Source/Programs/UnrealVS/UnrealVS.2019/Properties/AssemblyInfo.cs b/Engine/Source/Programs/UnrealVS/UnrealVS.2019/Properties/AssemblyInfo.cs
--- a/Engine/Source/Programs/UnrealVS/UnrealVS.2019/Properties/AssemblyInfo.cs
+++ b/Engine/Source/Programs/UnrealVS/UnrealVS.2019/Properties/AssemblyInfo.cs
@@ -31,8 +31,8 @@ using System.Runtime.InteropServices;
 // You can specify all the values or you can default the Revision and Build Numbers 
 // by using the '*' as shown below:
 
-[assembly: AssemblyVersion("1.79.0.0")]
-[assembly: AssemblyFileVersion("1.79.0.0")]
+[assembly: AssemblyVersion("1.80.0.0")]
+[assembly: AssemblyFileVersion("1.80.0.0")]
 
 
 
diff --git a/Engine/Source/Programs/UnrealVS/UnrealVS.2022/Properties/AssemblyInfo.cs b/Engine/Source/Programs/UnrealVS/UnrealVS.2022/Properties/AssemblyInfo.cs
deleted file mode 100644
diff --git a/Engine/Source/Programs/UnrealVS/UnrealVS.2022/UnrealVS.2022.csproj b/Engine/Source/Programs/UnrealVS/UnrealVS.2022/UnrealVS.2022.csproj
--- a/Engine/Source/Programs/UnrealVS/UnrealVS.2022/UnrealVS.2022.csproj
+++ b/Engine/Source/Programs/UnrealVS/UnrealVS.2022/UnrealVS.2022.csproj
@@ -46,9 +46,10 @@
     <BootstrapperEnabled>true</BootstrapperEnabled>
   </PropertyGroup>
   <ItemGroup>
-    <VSCTCompile Include="UnrealVS.vsct">
+    <VSCTCompile Include="..\UnrealVS.2019\UnrealVS.vsct">
       <ResourceName>Menus.ctmenu</ResourceName>
       <SubType>Designer</SubType>
+      <Link>UnrealVS.vsct</Link>
     </VSCTCompile>
   </ItemGroup>
   <ItemGroup>
@@ -66,12 +67,14 @@
     <Reference Include="WindowsBase" />
   </ItemGroup>
   <ItemGroup>
+    <Compile Include="..\UnrealVS.2019\Properties\AssemblyInfo.cs">
+      <Link>Properties\AssemblyInfo.cs</Link>
+    </Compile>
     <Compile Include="Resources.Designer.cs">
       <AutoGen>True</AutoGen>
       <DesignTime>True</DesignTime>
       <DependentUpon>Resources.resx</DependentUpon>
     </Compile>
-    <Compile Include="Properties\AssemblyInfo.cs" />
   </ItemGroup>
   <ItemGroup>
     <EmbeddedResource Include="Resources.resx">
@@ -92,10 +95,25 @@
     <None Include="Key.snk" />
   </ItemGroup>
   <ItemGroup>
+    <Resource Include="..\UnrealVS.2019\Resources\BatchBuilder.png">
+      <Link>Resources\BatchBuilder.png</Link>
+    </Resource>
+    <Resource Include="..\UnrealVS.2019\Resources\BatchBuilderBackground.png">
+      <Link>Resources\BatchBuilderBackground.png</Link>
+    </Resource>
+    <Resource Include="..\UnrealVS.2019\Resources\Icon.ico">
+      <Link>Resources\Icon.ico</Link>
+    </Resource>
+    <Resource Include="..\UnrealVS.2019\Resources\Icon.png">
+      <Link>Resources\Icon.png</Link>
+    </Resource>
+    <Resource Include="..\UnrealVS.2019\Resources\Preview.png">
+      <Link>Resources\Preview.png</Link>
+    </Resource>
+    <Resource Include="..\UnrealVS.2019\Resources\RefreshProjects.png">
+      <Link>Resources\RefreshProjects.png</Link>
+    </Resource>
     <None Include="source.extension.vsixmanifest" />
-    <None Include="Resources\BatchBuilder.png" />
-    <Resource Include="Resources\BatchBuilderBackground.png" />
-    <Resource Include="Resources\Icon.ico" />
     <Content Include="..\ReleaseNotes.txt">
       <Visible>false</Visible>
       <CopyToOutputDirectory>Always</CopyToOutputDirectory>
@@ -106,15 +124,6 @@
       <CopyToOutputDirectory>Always</CopyToOutputDirectory>
       <IncludeInVSIX>true</IncludeInVSIX>
     </Content>
-    <Content Include="Resources\Icon.png">
-      <CopyToOutputDirectory>Always</CopyToOutputDirectory>
-      <IncludeInVSIX>true</IncludeInVSIX>
-    </Content>
-    <Content Include="Resources\Preview.png">
-      <IncludeInVSIX>true</IncludeInVSIX>
-      <CopyToOutputDirectory>Always</CopyToOutputDirectory>
-    </Content>
-    <Resource Include="Resources\RefreshProjects.png" />
   </ItemGroup>
   <ItemGroup>
     <BootstrapperPackage Include=".NETFramework,Version=v4.5">
diff --git a/Engine/Source/Programs/UnrealVS/UnrealVS.Shared/CompileSingleFile.cs b/Engine/Source/Programs/UnrealVS/UnrealVS.Shared/CompileSingleFile.cs
--- a/Engine/Source/Programs/UnrealVS/UnrealVS.Shared/CompileSingleFile.cs
+++ b/Engine/Source/Programs/UnrealVS/UnrealVS.Shared/CompileSingleFile.cs
@@ -22,7 +22,7 @@ namespace UnrealVS
 		private const int UBTSubMenuID = 0x3103;
 		private string	  FileToCompileOriginalExt = "";
 
-		static readonly List<string> ValidExtensions = new List<string> { ".c", ".cc", ".cpp", ".cxx" };
+		static readonly HashSet<string> ValidExtensions = new HashSet<string>(StringComparer.OrdinalIgnoreCase) { ".c", ".cc", ".cpp", ".cxx" };
 
 		System.Diagnostics.Process ChildProcess;
 		private OleMenuCommand SubMenuCommand;
@@ -30,11 +30,13 @@ namespace UnrealVS
 		public CompileSingleFile()
 		{
 			CommandID CommandID = new CommandID(GuidList.UnrealVSCmdSet, CompileSingleFileButtonID);
-			MenuCommand CompileSingleFileButtonCommand = new MenuCommand(new EventHandler(CompileSingleFileButtonHandler), CommandID);
+			var CompileSingleFileButtonCommand = new OleMenuCommand(new EventHandler(CompileSingleFileButtonHandler), CommandID);
+			CompileSingleFileButtonCommand.BeforeQueryStatus += CompileSingleFileButtonCommand_BeforeQueryStatus;
 			UnrealVSPackage.Instance.MenuCommandService.AddCommand(CompileSingleFileButtonCommand);
 
 			CommandID CommandID2 = new CommandID(GuidList.UnrealVSCmdSet, PreprocessSingleFileButtonID);
-			MenuCommand PreprocessSingleFileButtonCommand = new MenuCommand(new EventHandler(CompileSingleFileButtonHandler), CommandID2);
+			var PreprocessSingleFileButtonCommand = new OleMenuCommand(new EventHandler(CompileSingleFileButtonHandler), CommandID2);
+			PreprocessSingleFileButtonCommand.BeforeQueryStatus += CompileSingleFileButtonCommand_BeforeQueryStatus;
 			UnrealVSPackage.Instance.MenuCommandService.AddCommand(PreprocessSingleFileButtonCommand);
 
 			CommandID CommandID3 = new CommandID(GuidList.UnrealVSCmdSet, CompileSingleModuleButtonID);
@@ -46,6 +48,27 @@ namespace UnrealVS
 			UnrealVSPackage.Instance.MenuCommandService.AddCommand(SubMenuCommand);
 		}
 
+		private void CompileSingleFileButtonCommand_BeforeQueryStatus(object sender, EventArgs e)
+		{
+			ThreadHelper.ThrowIfNotOnUIThread();
+			DTE DTE = UnrealVSPackage.Instance.DTE;
+			MenuCommand MenuCommand = sender as MenuCommand;
+			if (MenuCommand == null)
+			{
+				return;
+			}
+
+			if (DTE?.ActiveDocument == null)
+			{
+				MenuCommand.Enabled = false;
+				return;
+			}
+
+			// Check if the requested file is valid
+			string FileToCompileExt = Path.GetExtension(DTE.ActiveDocument.FullName);
+			MenuCommand.Enabled = ValidExtensions.Contains(FileToCompileExt);
+		}
+
 		public void Dispose()
 		{
 			KillChildProcess();
diff --git a/Engine/Source/Programs/UnrealVS/UnrealVS.Shared/FileBrowserWindowControl.xaml.cs b/Engine/Source/Programs/UnrealVS/UnrealVS.Shared/FileBrowserWindowControl.xaml.cs
--- a/Engine/Source/Programs/UnrealVS/UnrealVS.Shared/FileBrowserWindowControl.xaml.cs
+++ b/Engine/Source/Programs/UnrealVS/UnrealVS.Shared/FileBrowserWindowControl.xaml.cs
@@ -18,83 +18,107 @@ using System.Windows.Threading;
 
 namespace UnrealVS
 {
-    [global::System.Diagnostics.CodeAnalysis.SuppressMessageAttribute("", "VSTHRD010")]
-    public partial class FileBrowserWindowControl : UserControl
-    {
-        FileBrowserWindow Window;
-        bool ClearFilter;
-        bool IsRefreshing;
-        ListBox ActiveFilesListBox;
-
-        public class FileItem
-        {
-            public string Name { get; set; }
-            public string File { get; set; }
-            public string Project { get; set; }
-            public bool Selectable { get; set; }
-        }
+	[global::System.Diagnostics.CodeAnalysis.SuppressMessageAttribute("", "VSTHRD010")]
+	public partial class FileBrowserWindowControl : UserControl
+	{
+		FileBrowserWindow Window;
+		bool ClearFilter;
+		bool IsRefreshing;
+		ListBox ActiveFilesListBox;
+
+		public class FileItem
+		{
+			public string Name { get; set; }
+			public string File { get; set; }
+			public string Project { get; set; }
+			public bool Selectable { get; set; }
+		}
 
 		String UnrealVsFileName;
 
 		List<FileItem> AllFileItems;
-        List<FileItem> BookmarkedFileItems;
-        List<FileItem> RecentFileItems;
+		List<FileItem> BookmarkedFileItems;
+		List<FileItem> RecentFileItems;
 
-        Dictionary<string, FileItem> AllFileItemsLookup;
+		Dictionary<string, FileItem> AllFileItemsLookup;
 
 		public FileBrowserWindowControl(FileBrowserWindow window)
-        {
-            Window = window;
-            this.InitializeComponent();
-
-            AllFileItems = new List<FileItem>();
-            BookmarkedFileItems = new List<FileItem>();
-            RecentFileItems = new List<FileItem>();
-            IsVisibleChanged += FileBrowserWindowControl_IsVisibleChanged;
-            foreach (var listView in new[] { AllFilesListBox, BookmarkedFilesListBox, RecentFilesListBox })
-            {
-                listView.ItemsSource = new List<FileItem>();
-                listView.KeyDown += FileListView_KeyDown;
-                listView.PreviewKeyDown += FileListView_PreviewKeyDown;
-                listView.SelectionChanged += FileListView_SelectionChanged;
-                listView.GotFocus += FileListView_GotFocus;
-                listView.MouseDoubleClick += ListView_MouseDoubleClick;
-                listView.SelectionMode = SelectionMode.Extended;
-            }
-            FilterEditBox.TextChanged += FilterEditBox_TextChanged;
-            FilterEditBox.KeyDown += FilterEditBox_KeyDown;
-            FilterEditBox.PreviewKeyDown += FilterEditBox_PreviewKeyDown;
-            FilesListTab.SelectionChanged += FilesListTab_SelectionChanged;
-            FilesListTab.KeyDown += FilesListTab_KeyDown;
-            FilesListTab.SelectedIndex = 0;
+		{
+			Window = window;
+			this.InitializeComponent();
+
+			AllFileItems = new List<FileItem>();
+			BookmarkedFileItems = new List<FileItem>();
+			RecentFileItems = new List<FileItem>();
+			IsVisibleChanged += FileBrowserWindowControl_IsVisibleChanged;
+			foreach (var listView in new[] { AllFilesListBox, BookmarkedFilesListBox, RecentFilesListBox })
+			{
+				listView.ItemsSource = new List<FileItem>();
+				listView.KeyDown += FileListView_KeyDown;
+				listView.PreviewKeyDown += FileListView_PreviewKeyDown;
+				listView.SelectionChanged += FileListView_SelectionChanged;
+				listView.GotFocus += FileListView_GotFocus;
+				listView.MouseDoubleClick += ListView_MouseDoubleClick;
+				listView.SelectionMode = SelectionMode.Extended;
+			}
+			FilterEditBox.TextChanged += FilterEditBox_TextChanged;
+			FilterEditBox.KeyDown += FilterEditBox_KeyDown;
+			FilterEditBox.PreviewKeyDown += FilterEditBox_PreviewKeyDown;
+			DataObject.AddPastingHandler(FilterEditBox, FilterEditBox_OnPaste);
+			FilesListTab.SelectionChanged += FilesListTab_SelectionChanged;
+			FilesListTab.SelectedIndex = 0;
 
 			ActiveFilesListBox = AllFilesListBox;
-        }
-
-        private void FilesListTab_KeyDown(object sender, KeyEventArgs e)
-        {
-            /*
-            if (e.Key == Key.Right && FilesListTab.SelectedIndex < 2)
-                FilesListTab.SelectedIndex = (int)FilesListTab.SelectedIndex + 1;
-            else if (e.Key == Key.Left && FilesListTab.SelectedIndex > 0)
-                --FilesListTab.SelectedIndex;
-            */
-        }
-
-        private void FilesListTab_SelectionChanged(object sender, SelectionChangedEventArgs e)
-        {
-            switch (FilesListTab.SelectedIndex)
-            {
-                case 0:
-                    if (ActiveFilesListBox == AllFilesListBox)
-                        return;
-                    ActiveFilesListBox = AllFilesListBox;
-                    break;
-                case 1:
-                    if (ActiveFilesListBox == BookmarkedFilesListBox)
-                        return;
-                    ActiveFilesListBox = BookmarkedFilesListBox;
-                    break;
+			this.PreviewKeyDown += FileBrowserWindowControl_PreviewKeyDown;
+		}
+
+		private void FilterEditBox_OnPaste(object sender, DataObjectPastingEventArgs e)
+		{
+			var isText = e.SourceDataObject.GetDataPresent(DataFormats.UnicodeText, true);
+			if (!isText)
+				return;
+
+			var text = e.SourceDataObject.GetData(DataFormats.UnicodeText) as string;
+
+			if (string.IsNullOrEmpty(text) || text.IndexOfAny(Path.GetInvalidPathChars()) != -1)
+				return;
+
+			FilterEditBox.Text = Path.GetFileName(text);
+			e.CancelCommand();
+			FilterEditBox.CaretIndex = FilterEditBox.Text.Length;
+		}
+
+		[STAThread]
+		private void FileBrowserWindowControl_PreviewKeyDown(object sender, KeyEventArgs e)
+		{
+			if (FilterEditBox.IsKeyboardFocused)
+				return;
+
+			if (Keyboard.Modifiers.HasFlag(ModifierKeys.Control) && e.Key == Key.V || Keyboard.Modifiers.HasFlag(ModifierKeys.Shift) && e.Key == Key.Insert)
+			{
+				string text = Clipboard.GetText();
+
+				if (string.IsNullOrEmpty(text) || text.IndexOfAny(Path.GetInvalidPathChars()) != -1)
+					return;
+
+				FilterEditBox.Text = Path.GetFileName(text);
+			}
+		}
+
+		private void FilesListTab_SelectionChanged(object sender, SelectionChangedEventArgs e)
+		{
+			switch (FilesListTab.SelectedIndex)
+			{
+				case 0:
+					if (ActiveFilesListBox == AllFilesListBox)
+						return;
+					ActiveFilesListBox = AllFilesListBox;
+					break;
+				case 1:
+					if (ActiveFilesListBox == BookmarkedFilesListBox)
+						return;
+					ActiveFilesListBox = BookmarkedFilesListBox;
+					break;
 				case 2:
 					if (ActiveFilesListBox == RecentFilesListBox)
 						return;
@@ -102,889 +126,889 @@ namespace UnrealVS
 					break;
 			}
 			FilesListTab.UpdateLayout();
-            RefreshStatusBox();
-            AsyncFocusSelectedItem();
-        }
-
-        private void FilterEditBox_KeyDown(object sender, KeyEventArgs e)
-        {
-            if (e.Key == Key.Tab)
-            {
-                e.Handled = true;
-                FocusListViewItem(ActiveFilesListBox, ActiveFilesListBox.SelectedIndex);
-                return;
-            }
-        }
-        private void FilterEditBox_PreviewKeyDown(object sender, KeyEventArgs e)
-        {
-            var listBox = ActiveFilesListBox;
-            var getSelectedItem = new Func<FileItem>(() =>
-            {
-                if (listBox.SelectedItems.Count == 0)
-                    return (FileItem)listBox.Items[0];
-                else
-                    return (FileItem)listBox.SelectedItems[listBox.SelectedItems.Count - 1];
-            });
-
-            var moveFunc = new Action<int>((int steps) =>
-            {
-                FileItem selectedItem = getSelectedItem();
-                var index = listBox.Items.IndexOf(selectedItem);
-                var newIndex = Math.Min(Math.Max(0, index + steps), listBox.Items.Count - 1);
-                if (newIndex == index)
-                    return;
-                selectedItem = (FileItem)listBox.Items[newIndex];
-                if (selectedItem == null)
-                    return;
-                listBox.SelectedItems.Clear();
-                listBox.SelectedItems.Add(selectedItem);
-                listBox.ScrollIntoView(selectedItem);
-            });
-
-            if (e.Key == Key.Down)
-            {
-                moveFunc(1);
-                return;
-            }
-            if (e.Key == Key.Up)
-            {
-                moveFunc(-1);
-                return;
-            }
-            if (e.Key == Key.Enter)
-            {
-                if (OpenSelectedFile())
-                    HidePanel();
-                return;
-            }
-            if (e.Key == Key.PageDown)
-            {
-                e.Handled = true;
-                var lastVisible = GetVisibleListViewElement(listBox, (int)listBox.ActualHeight - 22);
-                if (lastVisible == null)
-                    return;
-                var lastItem = listBox.ItemContainerGenerator.ItemFromContainer(lastVisible);
-                FileItem selectedItem = getSelectedItem();
-                if (lastItem != selectedItem)
-                {
-                    listBox.SelectedItems.Clear();
-                    listBox.SelectedItems.Add(lastItem);
-                    listBox.ScrollIntoView(lastItem);
-                }
-                else
-                    moveFunc(((int)(listBox.ActualHeight / lastVisible.ActualHeight))-1);
-                return;
-            }
-            if (e.Key == Key.PageUp)
-            {
-                e.Handled = true;
-                var firstVisible = GetVisibleListViewElement(listBox, 10);
-                if (firstVisible == null)
-                    return;
-                var firstItem = listBox.ItemContainerGenerator.ItemFromContainer(firstVisible);
-                FileItem selectedItem = getSelectedItem();
-                if (firstItem != selectedItem)
-                {
-                    listBox.SelectedItems.Clear();
-                    listBox.SelectedItems.Add(firstItem);
-                    listBox.ScrollIntoView(firstItem);
-                }
-                else
-                    moveFunc(1 - (int)(listBox.ActualHeight / firstVisible.ActualHeight));
-            }
-        }
-
-        private void FilterEditBox_TextChanged(object sender, TextChangedEventArgs e)
-        {
-            RefreshListViews();
-            //AllFilesListBox.SelectedIndex = 0;
-        }
-
-        private void FileListView_GotFocus(object sender, RoutedEventArgs e)
-        {
-            ActiveFilesListBox = (ListBox)e.Source;
-        }
-
-        private void FileListView_SelectionChanged(object sender, SelectionChangedEventArgs e)
-        {
-            RefreshStatusBox();
-        }
-
-        private void FileListView_KeyDown(object sender, KeyEventArgs e)
-        {
-            if ((Keyboard.Modifiers & ModifierKeys.Control) != 0)
-            {
-                if (e.Key == Key.I)
-                    if (CreateIncludePath())
+			RefreshStatusBox();
+			AsyncFocusSelectedItem();
+		}
+
+		private void FilterEditBox_KeyDown(object sender, KeyEventArgs e)
+		{
+			if (e.Key == Key.Tab)
+			{
+				e.Handled = true;
+				FocusListViewItem(ActiveFilesListBox, ActiveFilesListBox.SelectedIndex);
+				return;
+			}
+		}
+		private void FilterEditBox_PreviewKeyDown(object sender, KeyEventArgs e)
+		{
+			var listBox = ActiveFilesListBox;
+			var getSelectedItem = new Func<FileItem>(() =>
+			{
+				if (listBox.SelectedItems.Count == 0)
+					return (FileItem)listBox.Items[0];
+				else
+					return (FileItem)listBox.SelectedItems[listBox.SelectedItems.Count - 1];
+			});
+
+			var moveFunc = new Action<int>((int steps) =>
+			{
+				FileItem selectedItem = getSelectedItem();
+				var index = listBox.Items.IndexOf(selectedItem);
+				var newIndex = Math.Min(Math.Max(0, index + steps), listBox.Items.Count - 1);
+				if (newIndex == index)
+					return;
+				selectedItem = (FileItem)listBox.Items[newIndex];
+				if (selectedItem == null)
+					return;
+				listBox.SelectedItems.Clear();
+				listBox.SelectedItems.Add(selectedItem);
+				listBox.ScrollIntoView(selectedItem);
+			});
+
+			if (e.Key == Key.Down)
+			{
+				moveFunc(1);
+				return;
+			}
+			if (e.Key == Key.Up)
+			{
+				moveFunc(-1);
+				return;
+			}
+			if (e.Key == Key.Enter)
+			{
+				if (OpenSelectedFile())
+					HidePanel();
+				return;
+			}
+			if (e.Key == Key.PageDown)
+			{
+				e.Handled = true;
+				var lastVisible = GetVisibleListViewElement(listBox, (int)listBox.ActualHeight - 22);
+				if (lastVisible == null)
+					return;
+				var lastItem = listBox.ItemContainerGenerator.ItemFromContainer(lastVisible);
+				FileItem selectedItem = getSelectedItem();
+				if (lastItem != selectedItem)
+				{
+					listBox.SelectedItems.Clear();
+					listBox.SelectedItems.Add(lastItem);
+					listBox.ScrollIntoView(lastItem);
+				}
+				else
+					moveFunc(((int)(listBox.ActualHeight / lastVisible.ActualHeight)) - 1);
+				return;
+			}
+			if (e.Key == Key.PageUp)
+			{
+				e.Handled = true;
+				var firstVisible = GetVisibleListViewElement(listBox, 10);
+				if (firstVisible == null)
+					return;
+				var firstItem = listBox.ItemContainerGenerator.ItemFromContainer(firstVisible);
+				FileItem selectedItem = getSelectedItem();
+				if (firstItem != selectedItem)
+				{
+					listBox.SelectedItems.Clear();
+					listBox.SelectedItems.Add(firstItem);
+					listBox.ScrollIntoView(firstItem);
+				}
+				else
+					moveFunc(1 - (int)(listBox.ActualHeight / firstVisible.ActualHeight));
+			}
+		}
+
+		private void FilterEditBox_TextChanged(object sender, TextChangedEventArgs e)
+		{
+			RefreshListViews();
+			//AllFilesListBox.SelectedIndex = 0;
+		}
+
+		private void FileListView_GotFocus(object sender, RoutedEventArgs e)
+		{
+			ActiveFilesListBox = (ListBox)e.Source;
+		}
+
+		private void FileListView_SelectionChanged(object sender, SelectionChangedEventArgs e)
+		{
+			RefreshStatusBox();
+		}
+
+		private void FileListView_KeyDown(object sender, KeyEventArgs e)
+		{
+			if ((Keyboard.Modifiers & ModifierKeys.Control) != 0)
+			{
+				if (e.Key == Key.I)
+					if (CreateIncludePath())
 						HidePanel();
 				return;
-            }
-            if (e.Key == Key.Left)
-            {
-                if (ActiveFilesListBox != AllFilesListBox)
-                    --FilesListTab.SelectedIndex;
-                //if (ActiveFilesListBox == BookmarkedFilesListBox)
-                //    ToggleListView(BookmarkedFilesListBox, AllFilesListBox);
-                return;
-            }
-
-            if (e.Key == Key.Right)
-            {
+			}
+			if (e.Key == Key.Left)
+			{
+				if (ActiveFilesListBox != AllFilesListBox)
+					--FilesListTab.SelectedIndex;
+				//if (ActiveFilesListBox == BookmarkedFilesListBox)
+				//    ToggleListView(BookmarkedFilesListBox, AllFilesListBox);
+				return;
+			}
+
+			if (e.Key == Key.Right)
+			{
 				if (ActiveFilesListBox != RecentFilesListBox)// && ((List<FileItem>)BookmarkedFilesListBox.ItemsSource).Count > 0)
 					++FilesListTab.SelectedIndex;
 				//if (ActiveFilesListBox == AllFilesListBox)// && ((List<FileItem>)BookmarkedFilesListBox.ItemsSource).Count > 0)
-                 //   FilesListTab.SelectedIndex = 1;
-                //if (ActiveFilesListBox == AllFilesListBox && ((List<FileItem>)BookmarkedFilesListBox.ItemsSource).Count > 0)
-               //     ToggleListView(AllFilesListBox, BookmarkedFilesListBox);
-                return;
-            }
-
-            if (e.Key == Key.Tab)
-            {
-                e.Handled = true;
-                FilterEditBox.Focus();
-                return;
-            }
-
-            if (e.Key == Key.Insert)
-            {
-                if (ActiveFilesListBox != AllFilesListBox)
-                    return;
-                bool BookmarksModified = false;
-                foreach (var i in ActiveFilesListBox.SelectedItems)
-                {
-                    var item = (FileItem)i;
-                    if (item == null || item.Name == "")
-                        continue;
-                    if (BookmarkedFileItems.Contains(item))
-                        continue;
-                    BookmarkedFileItems.Add(item);
-                    BookmarksModified = true;
-                }
-                if (!BookmarksModified)
-                    return;
-                RefreshListView(BookmarkedFilesListBox, BookmarkedFileItems);
-                //if (AllFilesListBox.SelectedIndex < AllFilesListBox.Items.Count - 1)
-                //    FocusListViewItem(AllFilesListBox, AllFilesListBox.SelectedIndex + 1);
+				//   FilesListTab.SelectedIndex = 1;
+				//if (ActiveFilesListBox == AllFilesListBox && ((List<FileItem>)BookmarkedFilesListBox.ItemsSource).Count > 0)
+				//     ToggleListView(AllFilesListBox, BookmarkedFilesListBox);
+				return;
+			}
+
+			if (e.Key == Key.Tab)
+			{
+				e.Handled = true;
+				FilterEditBox.Focus();
+				return;
+			}
+
+			if (e.Key == Key.Insert)
+			{
+				if (ActiveFilesListBox != AllFilesListBox)
+					return;
+				bool BookmarksModified = false;
+				foreach (var i in ActiveFilesListBox.SelectedItems)
+				{
+					var item = (FileItem)i;
+					if (item == null || item.Name == "")
+						continue;
+					if (BookmarkedFileItems.Contains(item))
+						continue;
+					BookmarkedFileItems.Add(item);
+					BookmarksModified = true;
+				}
+				if (!BookmarksModified)
+					return;
+				RefreshListView(BookmarkedFilesListBox, BookmarkedFileItems);
+				//if (AllFilesListBox.SelectedIndex < AllFilesListBox.Items.Count - 1)
+				//    FocusListViewItem(AllFilesListBox, AllFilesListBox.SelectedIndex + 1);
 				SaveSolutionSettings();
-                return;
-            }
-
-            if (e.Key == Key.Delete)
-            {
-                if (ActiveFilesListBox != BookmarkedFilesListBox)
-                    return;
-                bool BookmarksModified = false;
-
-                var index = BookmarkedFilesListBox.SelectedIndex;
-                foreach (var i in ActiveFilesListBox.SelectedItems)
-                {
-                    var item = (FileItem)i;
-                    if (item == null || item.Name == "")
-                        continue;
-                    if (!BookmarkedFileItems.Remove(item))
-                        continue;
-                    BookmarksModified = true;
-                }
-                if (BookmarksModified)
-                {
-                    RefreshListView(BookmarkedFilesListBox, BookmarkedFileItems);
-                    if (index == BookmarkedFilesListBox.Items.Count)
-                    {
-                        if (index == 0)
-                            FilesListTab.SelectedIndex = 0;//ActiveFilesListBox = AllFilesListBox;
-                        else
-                            index -= 1;
-                    }
-                    BookmarkedFilesListBox.SelectedIndex = index;
-                    AsyncFocusSelectedItem();
-                    SaveSolutionSettings();
-                }
-                return;
-            }
-
-            if (e.Key == Key.Enter)
-            {
-                if (OpenSelectedFile())
-                    HidePanel();
-                return;
-            }
-
-            bool resetSelection = false;
-
-            if (e.Key == Key.Back)
-            {
-                string str = FilterEditBox.Text;
-                if (string.IsNullOrEmpty(str))
-                    return;
-                if (ClearFilter)
-                {
-                    resetSelection = true;
-                    ClearFilter = false;
-                    FilterEditBox.Text = "";
-                }
-                else
-                    FilterEditBox.Text = str.Substring(0, str.Length - 1);
-            }
-            else
-            {
-                char c = GetCharFromKey(e.Key);
-                if (c == 0)
-                    return;
-                if (ClearFilter)
-                {
-                    resetSelection = true;
-                    ClearFilter = false;
-                    FilterEditBox.Text = "";
-                }
-                FilterEditBox.Text += c;
-                FilterEditBox.SelectAll();
-            }
-
-            if (resetSelection)
-            {
-                AllFilesListBox.SelectedIndex = 0;
-                BookmarkedFilesListBox.SelectedIndex = 0;
-                RecentFilesListBox.SelectedIndex = 0;
-            }
-
-            AsyncFocusSelectedItem();
-
-            e.Handled = true;
-        }
-
-        private void ListView_MouseDoubleClick(object sender, MouseButtonEventArgs e)
-        {
-            if (OpenSelectedFile())
-                HidePanel();
-        }
-
-        private bool OpenSelectedFile()
-        {
-            var toOpen = new List<string>();
-            foreach (var i in ActiveFilesListBox.SelectedItems)
-            {
-                var item = (FileItem)i;
-                if (item == null || item.Name == "")
-                    continue;
-                toOpen.Add(item.File);
-            }
-            ThreadHelper.JoinableTaskFactory.Run(async delegate
-            {
-                await ThreadHelper.JoinableTaskFactory.SwitchToMainThreadAsync();
-                foreach (var f in toOpen)
-                    UnrealVSPackage.Instance.DTE.ExecuteCommand("File.OpenFile", "\"" + f + "\"");
-            });
-            return toOpen.Count != 0;
-        }
-
-        private bool CreateIncludePath()
-        {
-            string includeStrings = "";
-            foreach (var i in ActiveFilesListBox.SelectedItems)
-            {
-                var item = (FileItem)i;
-                if (item == null || item.Name == "")
-                    continue;
-                string includePath = item.File;
-                includePath = includePath.Replace("\\", "/");
-                string includePathLwr = includePath.ToLower();
+				return;
+			}
+
+			if (e.Key == Key.Delete)
+			{
+				if (ActiveFilesListBox != BookmarkedFilesListBox)
+					return;
+				bool BookmarksModified = false;
+
+				var index = BookmarkedFilesListBox.SelectedIndex;
+				foreach (var i in ActiveFilesListBox.SelectedItems)
+				{
+					var item = (FileItem)i;
+					if (item == null || item.Name == "")
+						continue;
+					if (!BookmarkedFileItems.Remove(item))
+						continue;
+					BookmarksModified = true;
+				}
+				if (BookmarksModified)
+				{
+					RefreshListView(BookmarkedFilesListBox, BookmarkedFileItems);
+					if (index == BookmarkedFilesListBox.Items.Count)
+					{
+						if (index == 0)
+							FilesListTab.SelectedIndex = 0;//ActiveFilesListBox = AllFilesListBox;
+						else
+							index -= 1;
+					}
+					BookmarkedFilesListBox.SelectedIndex = index;
+					AsyncFocusSelectedItem();
+					SaveSolutionSettings();
+				}
+				return;
+			}
+
+			if (e.Key == Key.Enter)
+			{
+				if (OpenSelectedFile())
+					HidePanel();
+				return;
+			}
+
+			bool resetSelection = false;
+
+			if (e.Key == Key.Back)
+			{
+				string str = FilterEditBox.Text;
+				if (string.IsNullOrEmpty(str))
+					return;
+				if (ClearFilter)
+				{
+					resetSelection = true;
+					ClearFilter = false;
+					FilterEditBox.Text = "";
+				}
+				else
+					FilterEditBox.Text = str.Substring(0, str.Length - 1);
+			}
+			else
+			{
+				char c = GetCharFromKey(e.Key);
+				if (c == 0)
+					return;
+				if (ClearFilter)
+				{
+					resetSelection = true;
+					ClearFilter = false;
+					FilterEditBox.Text = "";
+				}
+				FilterEditBox.Text += c;
+				FilterEditBox.SelectAll();
+			}
+
+			if (resetSelection)
+			{
+				AllFilesListBox.SelectedIndex = 0;
+				BookmarkedFilesListBox.SelectedIndex = 0;
+				RecentFilesListBox.SelectedIndex = 0;
+			}
+
+			AsyncFocusSelectedItem();
+
+			e.Handled = true;
+		}
+
+		private void ListView_MouseDoubleClick(object sender, MouseButtonEventArgs e)
+		{
+			if (OpenSelectedFile())
+				HidePanel();
+		}
+
+		private bool OpenSelectedFile()
+		{
+			var toOpen = new List<string>();
+			foreach (var i in ActiveFilesListBox.SelectedItems)
+			{
+				var item = (FileItem)i;
+				if (item == null || item.Name == "")
+					continue;
+				toOpen.Add(item.File);
+			}
+			ThreadHelper.JoinableTaskFactory.Run(async delegate
+			{
+				await ThreadHelper.JoinableTaskFactory.SwitchToMainThreadAsync();
+				foreach (var f in toOpen)
+					UnrealVSPackage.Instance.DTE.ExecuteCommand("File.OpenFile", "\"" + f + "\"");
+			});
+			return toOpen.Count != 0;
+		}
+
+		private bool CreateIncludePath()
+		{
+			string includeStrings = "";
+			foreach (var i in ActiveFilesListBox.SelectedItems)
+			{
+				var item = (FileItem)i;
+				if (item == null || item.Name == "")
+					continue;
+				string includePath = item.File;
+				includePath = includePath.Replace("\\", "/");
+				string includePathLwr = includePath.ToLower();
 				var includeStartIndex = Math.Max(Math.Max(includePathLwr.LastIndexOf("/public/"), includePathLwr.LastIndexOf("/private/")), includePathLwr.LastIndexOf("/classes/"));
-                if (includeStartIndex != -1)
-                    includePath = includePath.Substring(includePath.IndexOf('/', includeStartIndex + 1) + 1);
-                includeStrings += $"#include \"{includePath}\"\r\n";
-            }
-            if (includeStrings == "")
-                return false;
-            Clipboard.SetText(includeStrings);
-            return true;
+				if (includeStartIndex != -1)
+					includePath = includePath.Substring(includePath.IndexOf('/', includeStartIndex + 1) + 1);
+				includeStrings += $"#include \"{includePath}\"\r\n";
+			}
+			if (includeStrings == "")
+				return false;
+			Clipboard.SetText(includeStrings);
+			return true;
 		}
 
 		private void ToggleListView(ListBox from, ListBox to)
-        {
-            var firstVisibleFrom = GetVisibleListViewElement(from, 10);
-            if (firstVisibleFrom == null)
-                return;
-            var firstVisibleFromIndex = from.ItemContainerGenerator.IndexFromContainer(firstVisibleFrom);
-            var firstVisibleTo = GetVisibleListViewElement(to, 10);
-            if (firstVisibleTo == null)
-                return;
-            var firstVisibleToIndex = to.ItemContainerGenerator.IndexFromContainer(firstVisibleTo);
-            var toSelectIndex = firstVisibleToIndex + from.SelectedIndex - firstVisibleFromIndex;
-            if (toSelectIndex >= to.Items.Count)
-                toSelectIndex = to.Items.Count - 1;
-            FocusListViewItem(to, toSelectIndex);
-        }
-
-        private void FocusListViewItem(ListBox listView, int itemIndex)
-        {
-            if (itemIndex == -1)
-                itemIndex = 0;
-            listView.SelectedIndex = itemIndex;
-            if (listView.Items.Count > 0)
-                (listView.ItemContainerGenerator.ContainerFromIndex(itemIndex) as ListBoxItem)?.Focus();
-            else
-                listView.Focus();
-            RefreshStatusBox();
-        }
-
-        private void AsyncFocusSelectedItem()
-        {
-            double interval = 0.1;
-            var timer = new DispatcherTimer(DispatcherPriority.Normal);
-            timer.Tick += (s, e) =>
-            {
-                var items = (List<FileItem>)ActiveFilesListBox.ItemsSource;
-                if (items.Count == 0)
-                {
-                    ActiveFilesListBox.Focus();
-                    timer.Stop();
-                    return;
-                }
-                var item = (ActiveFilesListBox.ItemContainerGenerator.ContainerFromIndex(Math.Max(0, ActiveFilesListBox.SelectedIndex)) as ListBoxItem);
-                if (item == null)
-                {
-                    timer.Interval = TimeSpan.FromSeconds(interval);
-                    interval = Math.Min(interval + 0.1, 1.0);
-                    return;
-                }
-                item.Focus();
-                timer.Stop();
-            };
-            timer.Start();
-        }
-
-        private void RefreshListViews()
-        {
-            RefreshListView(AllFilesListBox, AllFileItems);
+		{
+			var firstVisibleFrom = GetVisibleListViewElement(from, 10);
+			if (firstVisibleFrom == null)
+				return;
+			var firstVisibleFromIndex = from.ItemContainerGenerator.IndexFromContainer(firstVisibleFrom);
+			var firstVisibleTo = GetVisibleListViewElement(to, 10);
+			if (firstVisibleTo == null)
+				return;
+			var firstVisibleToIndex = to.ItemContainerGenerator.IndexFromContainer(firstVisibleTo);
+			var toSelectIndex = firstVisibleToIndex + from.SelectedIndex - firstVisibleFromIndex;
+			if (toSelectIndex >= to.Items.Count)
+				toSelectIndex = to.Items.Count - 1;
+			FocusListViewItem(to, toSelectIndex);
+		}
+
+		private void FocusListViewItem(ListBox listView, int itemIndex)
+		{
+			if (itemIndex == -1)
+				itemIndex = 0;
+			listView.SelectedIndex = itemIndex;
+			if (listView.Items.Count > 0)
+				(listView.ItemContainerGenerator.ContainerFromIndex(itemIndex) as ListBoxItem)?.Focus();
+			else
+				listView.Focus();
+			RefreshStatusBox();
+		}
+
+		private void AsyncFocusSelectedItem()
+		{
+			double interval = 0.1;
+			var timer = new DispatcherTimer(DispatcherPriority.Normal);
+			timer.Tick += (s, e) =>
+			{
+				var items = (List<FileItem>)ActiveFilesListBox.ItemsSource;
+				if (items.Count == 0)
+				{
+					ActiveFilesListBox.Focus();
+					timer.Stop();
+					return;
+				}
+				var item = (ActiveFilesListBox.ItemContainerGenerator.ContainerFromIndex(Math.Max(0, ActiveFilesListBox.SelectedIndex)) as ListBoxItem);
+				if (item == null)
+				{
+					timer.Interval = TimeSpan.FromSeconds(interval);
+					interval = Math.Min(interval + 0.1, 1.0);
+					return;
+				}
+				item.Focus();
+				timer.Stop();
+			};
+			timer.Start();
+		}
+
+		private void RefreshListViews()
+		{
+			RefreshListView(AllFilesListBox, AllFileItems);
 			RefreshListView(BookmarkedFilesListBox, BookmarkedFileItems);
 			RefreshListView(RecentFilesListBox, RecentFileItems);
 		}
 
-        struct Score
-        {
-            public int Sorting;
-            public int Index;
-        }
+		struct Score
+		{
+			public int Sorting;
+			public int Index;
+		}
 
 		private void RefreshListView(ListBox listView, List<FileItem> source)
-        {
-            var fileItems = (List<FileItem>)listView.ItemsSource;
-            fileItems.Clear();
-
-            string filterStr = FilterEditBox.Text;
-            var filter = filterStr.Split(new char[] { ' ' }, StringSplitOptions.RemoveEmptyEntries);
-
-            var selections = new HashSet<FileItem>();
-            foreach (var i in listView.SelectedItems)
-            {
-                var item = (FileItem)i;
-                if (item != null && item.Name != "")
-                    selections.Add(item);
-            }
-
-            var filterLen = filter.Length;
-            if (filterLen > 0)
-            {
-                var filterSortedItems = new List<ValueTuple<Score, FileItem>>();
-                var scoreArray = new ValueTuple<int, int>[5];
-                int ItemIndex = 0;
-                foreach (var item in source)
-                {
-                    int scoreArrayIndex = 0;
-                    for (int i = 0; i != filterLen; ++i)
-                    {
-                        var f = filter[i];
-                       int index = item.Name.IndexOf(f, System.StringComparison.OrdinalIgnoreCase);
-                        if (index == -1)
-                        {
-                            scoreArrayIndex = 0;
-                            break;
-                        }
-                        scoreArray[scoreArrayIndex].Item1 = index;
-                        scoreArray[scoreArrayIndex].Item2 = scoreArrayIndex;
-                        ++scoreArrayIndex;
-                        if (scoreArrayIndex == 5)
-                            break;
-                    }
-
-                    if (scoreArrayIndex == 0)
-                        continue;
-
-                    // We want to sort the ones that match the filter orders first and prioritize the ones that starts with the first filter entry
-                    int sortingScore = 0;
-                    Array.Sort(scoreArray, 0, scoreArrayIndex);
-                    for (int i = 0; i != scoreArrayIndex; ++i)
-                    {
-                        var Tuple = scoreArray[i];
-                        if (i == 0 && Tuple.Item1 == 0 && Tuple.Item2 == 0)
-                            sortingScore -= 100;
-                        sortingScore += (10 >> Tuple.Item2) * i + Tuple.Item1;
-                    }
-                    var s = new Score() { Sorting = sortingScore, Index = ItemIndex++ };
-                    filterSortedItems.Add(new ValueTuple<Score, FileItem>(s, item));
-                }
-
-                filterSortedItems.Sort((a, b) =>
-                {
-                    if (a.Item1.Sorting != b.Item1.Sorting)
-                        return a.Item1.Sorting - b.Item1.Sorting;
-                    return a.Item1.Index - b.Item1.Index;
-                });
-                foreach (var tuple in filterSortedItems)
-                    fileItems.Add(tuple.Item2);
-            }
-            else
-            {
-                fileItems.AddRange(source);
-            }
-
-            int fileCount = fileItems.Count;
-
-            // If empty we need to add one entry to make sure listview can be focused in windows
-            if (fileItems.Count == 0)
-                fileItems.Add(new FileItem() { Name = "" });
-
-            CollectionViewSource.GetDefaultView(fileItems)?.Refresh();
-
-            if (selections.Count != 0)
-            {
-                foreach (var item in fileItems)
-                {
-                    if (!selections.Contains(item))
-                        continue;
-                    listView.SelectedItems.Add(item);
-                    listView.ScrollIntoView(item);
-                }
-            }
-            
-            if (listView.SelectedItems.Count == 0)
-                listView.SelectedItems.Add(fileItems[0]);
-
-
-            if (listView == AllFilesListBox)
-                AllFilesTab.Header = $"All Files ({fileCount}/{source.Count})";
-            else if (listView == BookmarkedFilesListBox)
-                BookmarkedFilesTab.Header = $"Bookmarked Files ({fileCount}/{source.Count})";
-			else 
+		{
+			var fileItems = (List<FileItem>)listView.ItemsSource;
+			fileItems.Clear();
+
+			string filterStr = FilterEditBox.Text;
+			var filter = filterStr.Split(new char[] { ' ' }, StringSplitOptions.RemoveEmptyEntries);
+
+			var selections = new HashSet<FileItem>();
+			foreach (var i in listView.SelectedItems)
+			{
+				var item = (FileItem)i;
+				if (item != null && item.Name != "")
+					selections.Add(item);
+			}
+
+			var filterLen = filter.Length;
+			if (filterLen > 0)
+			{
+				var filterSortedItems = new List<ValueTuple<Score, FileItem>>();
+				var scoreArray = new ValueTuple<int, int>[5];
+				int ItemIndex = 0;
+				foreach (var item in source)
+				{
+					int scoreArrayIndex = 0;
+					for (int i = 0; i != filterLen; ++i)
+					{
+						var f = filter[i];
+						int index = item.Name.IndexOf(f, System.StringComparison.OrdinalIgnoreCase);
+						if (index == -1)
+						{
+							scoreArrayIndex = 0;
+							break;
+						}
+						scoreArray[scoreArrayIndex].Item1 = index;
+						scoreArray[scoreArrayIndex].Item2 = scoreArrayIndex;
+						++scoreArrayIndex;
+						if (scoreArrayIndex == 5)
+							break;
+					}
+
+					if (scoreArrayIndex == 0)
+						continue;
+
+					// We want to sort the ones that match the filter orders first and prioritize the ones that starts with the first filter entry
+					int sortingScore = 0;
+					Array.Sort(scoreArray, 0, scoreArrayIndex);
+					for (int i = 0; i != scoreArrayIndex; ++i)
+					{
+						var Tuple = scoreArray[i];
+						if (i == 0 && Tuple.Item1 == 0 && Tuple.Item2 == 0)
+							sortingScore -= 100;
+						sortingScore += (10 >> Tuple.Item2) * i + Tuple.Item1;
+					}
+					var s = new Score() { Sorting = sortingScore, Index = ItemIndex++ };
+					filterSortedItems.Add(new ValueTuple<Score, FileItem>(s, item));
+				}
+
+				filterSortedItems.Sort((a, b) =>
+				{
+					if (a.Item1.Sorting != b.Item1.Sorting)
+						return a.Item1.Sorting - b.Item1.Sorting;
+					return a.Item1.Index - b.Item1.Index;
+				});
+				foreach (var tuple in filterSortedItems)
+					fileItems.Add(tuple.Item2);
+			}
+			else
+			{
+				fileItems.AddRange(source);
+			}
+
+			int fileCount = fileItems.Count;
+
+			// If empty we need to add one entry to make sure listview can be focused in windows
+			if (fileItems.Count == 0)
+				fileItems.Add(new FileItem() { Name = "" });
+
+			CollectionViewSource.GetDefaultView(fileItems)?.Refresh();
+
+			if (selections.Count != 0)
+			{
+				foreach (var item in fileItems)
+				{
+					if (!selections.Contains(item))
+						continue;
+					listView.SelectedItems.Add(item);
+					listView.ScrollIntoView(item);
+				}
+			}
+
+			if (listView.SelectedItems.Count == 0)
+				listView.SelectedItems.Add(fileItems[0]);
+
+
+			if (listView == AllFilesListBox)
+				AllFilesTab.Header = $"All Files ({fileCount}/{source.Count})";
+			else if (listView == BookmarkedFilesListBox)
+				BookmarkedFilesTab.Header = $"Bookmarked Files ({fileCount}/{source.Count})";
+			else
 				RecentFilesTab.Header = $"Recent Files ({fileCount}/{source.Count})";
 		}
 
 		private void RefreshStatusBox()
-        {
-            var items = ActiveFilesListBox.SelectedItems;
-            if (items.Count == 0)
-                StatusText.Text = "";
-            else if (items.Count == 1)
-            {
-                var item = (FileItem)items[0];
-                StatusText.Text = item.File;
-            }
-            else
-                StatusText.Text = "Multiple files selected";
-        }
-
-        private void FileListView_PreviewKeyDown(object sender, KeyEventArgs e)
-        {
-            if (e.Key != Key.Space)
-                return;
-            if (ClearFilter)
-            {
-                ClearFilter = false;
-                return;
-            }
-            FilterEditBox.Text += ' ';
-            e.Handled = true;
-        }
-
-        internal void HandleEscape()
-        {
-            if (HelpDialog.Visibility == Visibility.Visible)
-                HelpDialog.Visibility = Visibility.Collapsed;
-            else
-                HidePanel();
-        }
-
-        internal void HandleF1()
-        {
-            HelpDialog.Visibility = Visibility.Visible;
-        }
-
-        internal void HandleF5()
-        {
-            AsyncRefreshFileList();
-        }
-
-        internal void HandleSolutionChanged()
-        {
-            SaveSolutionSettings();
-            if (IsVisible)
-                AsyncRefreshFileList();
-            else
-                AllFileItems.Clear(); // Will trigger AsyncRefreshFileList when visible
-        }
-
-        internal void HandleDocumentActivated(Document Document)
-		{
-            var fileName = Document.FullName;
-            FileItem fileItem;
-            if (AllFileItemsLookup == null || !AllFileItemsLookup.TryGetValue(fileName, out fileItem))
-                return;
-
-            RecentFileItems.RemoveAll((item) => item == fileItem);
-            RecentFileItems.Insert(0, fileItem);
-
-            int RecentMaxCount = 100;
-            if (RecentFileItems.Count > RecentMaxCount)
-                RecentFileItems.RemoveRange(RecentMaxCount, RecentFileItems.Count - RecentMaxCount);
+		{
+			var items = ActiveFilesListBox.SelectedItems;
+			if (items.Count == 0)
+				StatusText.Text = "";
+			else if (items.Count == 1)
+			{
+				var item = (FileItem)items[0];
+				StatusText.Text = item.File;
+			}
+			else
+				StatusText.Text = "Multiple files selected";
+		}
+
+		private void FileListView_PreviewKeyDown(object sender, KeyEventArgs e)
+		{
+			if (e.Key != Key.Space)
+				return;
+			if (ClearFilter)
+			{
+				ClearFilter = false;
+				return;
+			}
+			FilterEditBox.Text += ' ';
+			e.Handled = true;
+		}
+
+		internal void HandleEscape()
+		{
+			if (HelpDialog.Visibility == Visibility.Visible)
+				HelpDialog.Visibility = Visibility.Collapsed;
+			else
+				HidePanel();
+		}
+
+		internal void HandleF1()
+		{
+			HelpDialog.Visibility = Visibility.Visible;
+		}
+
+		internal void HandleF5()
+		{
+			AsyncRefreshFileList();
+		}
+
+		internal void HandleSolutionChanged()
+		{
+			SaveSolutionSettings();
+			if (IsVisible)
+				AsyncRefreshFileList();
+			else
+				AllFileItems.Clear(); // Will trigger AsyncRefreshFileList when visible
+		}
+
+		internal void HandleDocumentActivated(Document Document)
+		{
+			var fileName = Document.FullName;
+			FileItem fileItem;
+			if (AllFileItemsLookup == null || !AllFileItemsLookup.TryGetValue(fileName, out fileItem))
+				return;
+
+			RecentFileItems.RemoveAll((item) => item == fileItem);
+			RecentFileItems.Insert(0, fileItem);
+
+			int RecentMaxCount = 100;
+			if (RecentFileItems.Count > RecentMaxCount)
+				RecentFileItems.RemoveRange(RecentMaxCount, RecentFileItems.Count - RecentMaxCount);
 
 			RefreshListView(RecentFilesListBox, RecentFileItems);
 		}
 
 		private void AsyncRefreshFileList()
-        {
-            if (IsRefreshing)
-                return;
-            IsRefreshing = true;
-            RefreshingText.Text = "Refreshing File Lists (0)";
-            RefreshingDialog.Visibility = Visibility.Visible;
-
-            DispatcherTimer timer = new DispatcherTimer(DispatcherPriority.Render);
-            timer.Interval = TimeSpan.FromSeconds(0.2);
-            var traverser = new SolutionTraverser();
-            timer.Tick += (s, e_) =>
-            {
-                var listItems = traverser.Update();
-                if (listItems == null)
-                {
-                    RefreshingText.Text = $"Refreshing File Lists ({traverser.HandledItemCount})";
-                    timer.Interval = TimeSpan.FromSeconds(0.01);
-                    return;
-                }
-                AllFileItems = listItems;
+		{
+			if (IsRefreshing)
+				return;
+			IsRefreshing = true;
+			RefreshingText.Text = "Refreshing File Lists (0)";
+			RefreshingDialog.Visibility = Visibility.Visible;
+
+			DispatcherTimer timer = new DispatcherTimer(DispatcherPriority.Render);
+			timer.Interval = TimeSpan.FromSeconds(0.2);
+			var traverser = new SolutionTraverser();
+			timer.Tick += (s, e_) =>
+			{
+				var listItems = traverser.Update();
+				if (listItems == null)
+				{
+					RefreshingText.Text = $"Refreshing File Lists ({traverser.HandledItemCount})";
+					timer.Interval = TimeSpan.FromSeconds(0.01);
+					return;
+				}
+				AllFileItems = listItems;
 
 				var lookup = new Dictionary<string, FileItem>();
 				foreach (var item in AllFileItems)
 					try { lookup.Add(item.File, item); } catch (Exception) { }
-                AllFileItemsLookup = lookup;
+				AllFileItemsLookup = lookup;
 
 				LoadSolutionSettings();
 
-                //ClearFilter = true;
-                RefreshListViews();
+				//ClearFilter = true;
+				RefreshListViews();
 
-                //AllFilesListBox.SelectedIndex = 0;
+				//AllFilesListBox.SelectedIndex = 0;
 				//BookmarkedFilesListBox.SelectedIndex = 0;
 				//RecentFilesListBox.SelectedIndex = 0;
 				AsyncFocusSelectedItem();
 
-                timer.Stop();
-                RefreshingDialog.Visibility = Visibility.Collapsed;
-                IsRefreshing = false;
-            };
-            timer.Start();
-        }
-
-        private void FileBrowserWindowControl_IsVisibleChanged(object sender, DependencyPropertyChangedEventArgs e)
-        {
-            if (!(bool)e.NewValue)
-                return;
-            FilterEditBox.SelectAll();
-            ClearFilter = true;
-            //FilesListTab.SelectedIndex = 0;
-
-            if (AllFileItems.Count == 0)
-            {
-                RefreshListViews(); // To make sure we get one invisible entry
-                AsyncRefreshFileList();
-            }
-
-            AsyncFocusSelectedItem();
-        }
-
-        private void HidePanel()
-        {
-            ThreadHelper.ThrowIfNotOnUIThread();
-            IVsWindowFrame ToolWindowFrame = (IVsWindowFrame)Window.Frame;
-            ToolWindowFrame.Hide();
-        }
-
-        private class FileBrowserSettings
-        {
-            public List<string> Bookmarks { get; set; } = new List<string>();
+				timer.Stop();
+				RefreshingDialog.Visibility = Visibility.Collapsed;
+				IsRefreshing = false;
+			};
+			timer.Start();
+		}
+
+		private void FileBrowserWindowControl_IsVisibleChanged(object sender, DependencyPropertyChangedEventArgs e)
+		{
+			if (!(bool)e.NewValue)
+				return;
+			FilterEditBox.SelectAll();
+			ClearFilter = true;
+			//FilesListTab.SelectedIndex = 0;
+
+			if (AllFileItems.Count == 0)
+			{
+				RefreshListViews(); // To make sure we get one invisible entry
+				AsyncRefreshFileList();
+			}
+
+			AsyncFocusSelectedItem();
+		}
+
+		private void HidePanel()
+		{
+			ThreadHelper.ThrowIfNotOnUIThread();
+			IVsWindowFrame ToolWindowFrame = (IVsWindowFrame)Window.Frame;
+			ToolWindowFrame.Hide();
+		}
+
+		private class FileBrowserSettings
+		{
+			public List<string> Bookmarks { get; set; } = new List<string>();
 			public List<string> Recents { get; set; } = new List<string>();
 		}
 
 
 		private void SaveSolutionSettings()
-        {
-            if (String.IsNullOrEmpty(UnrealVsFileName))
-                return;
-            using (var file = File.CreateText(UnrealVsFileName))
-            {
-                var settings = new FileBrowserSettings();
+		{
+			if (String.IsNullOrEmpty(UnrealVsFileName))
+				return;
+			using (var file = File.CreateText(UnrealVsFileName))
+			{
+				var settings = new FileBrowserSettings();
 				foreach (var item in BookmarkedFileItems)
 					settings.Bookmarks.Add(item.File);
 				foreach (var item in RecentFileItems)
 					settings.Recents.Add(item.File);
 				string json = JsonSerializer.Serialize(settings, new JsonSerializerOptions { WriteIndented = true });
-                file.Write(json);
-            }
-        }
-
-        private void LoadSolutionSettings()
-        {
-            var solutionFileName = UnrealVSPackage.Instance.DTE.Solution.FileName;
-            if (String.IsNullOrEmpty(solutionFileName))
-            {
-                BookmarkedFileItems.Clear();
-                RecentFileItems.Clear();
-                UnrealVsFileName = null;
+				file.Write(json);
+			}
+		}
+
+		private void LoadSolutionSettings()
+		{
+			var solutionFileName = UnrealVSPackage.Instance.DTE.Solution.FileName;
+			if (String.IsNullOrEmpty(solutionFileName))
+			{
+				BookmarkedFileItems.Clear();
+				RecentFileItems.Clear();
+				UnrealVsFileName = null;
+				return;
+			}
+			UnrealVsFileName = solutionFileName.Substring(0, solutionFileName.Length - 3) + "unrealvs";
+			if (!File.Exists(UnrealVsFileName))
+				return;
+			string json = File.ReadAllText(UnrealVsFileName);
+			if (String.IsNullOrEmpty(json))
+				return;
+			var settings = JsonSerializer.Deserialize<FileBrowserSettings>(json, new JsonSerializerOptions { PropertyNameCaseInsensitive = true });
+			if (settings == null)
+				return;
+
+
+			BookmarkedFileItems.Clear();
+			if (settings.Bookmarks.Count == 0 && settings.Recents.Count == 0)
 				return;
-            }
-            UnrealVsFileName = solutionFileName.Substring(0, solutionFileName.Length - 3) + "unrealvs";
-            if (!File.Exists(UnrealVsFileName))
-                return;
-            string json = File.ReadAllText(UnrealVsFileName);
-            if (String.IsNullOrEmpty(json))
-                return;
-            var settings = JsonSerializer.Deserialize<FileBrowserSettings>(json, new JsonSerializerOptions { PropertyNameCaseInsensitive = true });
-            if (settings == null)
-                return;
-
-
-            BookmarkedFileItems.Clear();
-            if (settings.Bookmarks.Count == 0 && settings.Recents.Count == 0)
-                return;
-
-            var fileItems = new[] { BookmarkedFileItems, RecentFileItems };
-            var settingItems = new[] { settings.Bookmarks, settings.Recents };
+
+			var fileItems = new[] { BookmarkedFileItems, RecentFileItems };
+			var settingItems = new[] { settings.Bookmarks, settings.Recents };
 			FileItem fileItem;
 
-            for (int i = 0; i != 2; ++i)
-            {
-                var added = new HashSet<FileItem>();
-                foreach (var item in settingItems[i])
-                    if (AllFileItemsLookup.TryGetValue(item, out fileItem))
-                        if (added.Add(fileItem))
-                            fileItems[i].Add(fileItem);
-            }
-        }
-
-        private static ListBoxItem GetVisibleListViewElement(ListBox listView, int y)
-        {
-            HitTestResult hitTest = VisualTreeHelper.HitTest(listView, new Point(10, y));
-            DependencyObject depObj = hitTest.VisualHit as DependencyObject;
-            if (depObj == null)
-                return null;
-            DependencyObject current = depObj;
-            while (current != null && current != listView)
-            {
-                ListBoxItem listBoxItem = current as ListBoxItem;
-                if (listBoxItem != null)
-                    return listBoxItem;
-                current = VisualTreeHelper.GetParent(current);
-            }
-            return null;
-        }
-        private static char GetCharFromKey(Key key)
-        {
-            char ch = '\0';
-
-            int virtualKey = KeyInterop.VirtualKeyFromKey(key);
-            byte[] keyboardState = new byte[256];
-            NativeMethods.GetKeyboardState(keyboardState);
-
-            uint scanCode = NativeMethods.MapVirtualKey((uint)virtualKey, NativeMethods.MapType.MAPVK_VK_TO_VSC);
-            StringBuilder stringBuilder = new StringBuilder(2);
-
-            int result = NativeMethods.ToUnicode((uint)virtualKey, scanCode, keyboardState, stringBuilder, stringBuilder.Capacity, 0);
-            switch (result)
-            {
-                case -1:
-                    break;
-                case 0:
-                    break;
-                case 1:
-                default:
-                    ch = stringBuilder[0];
-                    break;
-            }
-            return ch;
-        }
-
-        class StackItem
-        {
-            public ProjectItems Items;
-            public int Index;
-            public Project Project;
-        }
-
-        class SolutionTraverser
-        {
-            private SortedDictionary<string, FileItem> SortedFileItems = new SortedDictionary<string, FileItem>();
-            private Dictionary<string, Dictionary<string, FileItem>> CollidingFileItems = new Dictionary<string, Dictionary<string, FileItem>>();
-            private Dictionary<string, string> ModuleRoots = new Dictionary<string, string>();
-
-            private Stack<StackItem> ItemStack = new Stack<StackItem>();
-            private int ProjectIndex;
-            public int HandledItemCount { get; private set; }
-
-            public List<FileItem> Update()
-            {
-                Stack<StackItem> itemStack = ItemStack;
-                var projects = UnrealVSPackage.Instance.DTE.Solution.Projects;
-                var projectCount = projects.Count;
-
-                int traverseCounter = 0;
-
-                while (ProjectIndex < projectCount)
-                {
-                    var project = projects.Item(ProjectIndex + 1);
-
-                    StackItem stackItem;
-                    if (ItemStack.Count > 0)
-                        stackItem = ItemStack.Pop();
-                    else
-                        stackItem = new StackItem() { Items = project.ProjectItems, Project = project };
-
-                    while (true)
-                    {
-                        if (traverseCounter > 5000)
-                        {
-                            ItemStack.Push(stackItem);
-                            return null;
-                        }
-
-                        if (stackItem.Items == null || stackItem.Index == stackItem.Items.Count)
-                        {
-                            if (itemStack.Count == 0)
-                                break;
-                            stackItem = itemStack.Pop();
-                            ++stackItem.Index;
-                            continue;
-                        }
-
-                        var projectItem = stackItem.Items.Item(stackItem.Index + 1);
-                        ++traverseCounter;
-
-                        if (projectItem.FileCount != 0)
-                        {
-                            var file = projectItem.FileNames[1];
-                            if (file != null)
-                            {
-                                var name = projectItem.Name;
-                                if (name != file)
-                                    if (!file.EndsWith("\\")) // Skip folders
-                                    {
-                                        FileItem fileItem = new FileItem() { Name = name, File = file, Project = stackItem.Project.Name };
-                                        FileItem existingFileItem;
-                                        if (SortedFileItems.TryGetValue(name, out existingFileItem))
-                                        {
-                                            if (existingFileItem.File != file)
-                                            {
-                                                Dictionary<string, FileItem> colList;
-                                                if (!CollidingFileItems.TryGetValue(name, out colList))
-                                                {
-                                                    colList = new Dictionary<string, FileItem>();
-                                                    colList.Add(existingFileItem.File, existingFileItem);
-                                                    CollidingFileItems[name] = colList;
-                                                }
-                                                if (!colList.ContainsKey(fileItem.File))
-                                                    colList.Add(fileItem.File, fileItem);
-                                            }
-                                        }
-                                        else
-                                            SortedFileItems.Add(name, fileItem);
-
-                                        if (file.EndsWith(".build.cs", StringComparison.OrdinalIgnoreCase))
-                                        {
-                                            var lastBackslash = file.LastIndexOf('\\');
-                                            var moduleDir = file.Substring(0, lastBackslash);
-                                            var moduleName = file.Substring(lastBackslash + 1, file.Length - lastBackslash - 10);
-                                            if (!ModuleRoots.ContainsKey(moduleDir))
-                                                ModuleRoots.Add(moduleDir, moduleName);
-                                            //else if (ModuleRoots[moduleDir] != moduleName)
-                                            //    Logging.WriteLine(file);
-                                        }
-
-                                        ++HandledItemCount;
-                                    }
-                            }
-                        }
-
-                        if (projectItem.SubProject != null)
-                        {
-                            itemStack.Push(stackItem);
-                            stackItem = new StackItem() { Items = projectItem.SubProject.ProjectItems, Project = projectItem.SubProject };
-                            continue;
-                        }
-
-                        if (projectItem.ProjectItems != null)
-                        {
-                            itemStack.Push(stackItem);
-                            stackItem = new StackItem() { Items = projectItem.ProjectItems, Project = stackItem.Project };
-                            continue;
-                        }
-
-                        ++stackItem.Index;
-                    }
-
-                    ++ProjectIndex;
-                }
-
-                // We need to rename all collisions
-                foreach (var kv in CollidingFileItems)
-                {
-                    SortedFileItems.Remove(kv.Key);
-                    var usedParents = new Dictionary<string, int>();
-
-                    foreach (var itemKv in kv.Value)
-                    {
-                        var item = itemKv.Value;
-                        string parent = item.Project;
-                        var path = item.File;
-                        int lastBackslash = path.LastIndexOf('\\');
-                        while (lastBackslash != -1)
-                        {
-                            path = path.Substring(0, lastBackslash);
-                            string moduleName;
-                            if (ModuleRoots.TryGetValue(path, out moduleName))
-                            {
-                                parent = moduleName;
-                                break;
-                            }
-                            lastBackslash = path.LastIndexOf('\\');
-                        }
-
-                        string sortName = item.Name + $" ({parent})";
-
-                        int counter = 0;
-                        if (usedParents.TryGetValue(parent, out counter))
-                        {
-                            if (counter == 0)
-                            {
-                                var alreadyAddedItem = SortedFileItems[sortName];
-                                SortedFileItems.Remove(sortName);
-                                alreadyAddedItem.Name = item.Name + $" ({parent} 0)";
-                                SortedFileItems.Add(alreadyAddedItem.Name, alreadyAddedItem);
-                            }
-                            ++counter;
-                            usedParents[parent] = counter;
-                            item.Name = item.Name + $" ({parent} {counter})";
-                            SortedFileItems.Add(item.Name, item);
-                        }
-                        else
-                        {
-                            usedParents.Add(parent, 0);
-                            item.Name = sortName;
-                            SortedFileItems.Add(item.Name, item);
-                        }
-                    }
-                }
-
-                return new List<FileItem>(SortedFileItems.Values);
-            }
-        }
-    }
+			for (int i = 0; i != 2; ++i)
+			{
+				var added = new HashSet<FileItem>();
+				foreach (var item in settingItems[i])
+					if (AllFileItemsLookup.TryGetValue(item, out fileItem))
+						if (added.Add(fileItem))
+							fileItems[i].Add(fileItem);
+			}
+		}
+
+		private static ListBoxItem GetVisibleListViewElement(ListBox listView, int y)
+		{
+			HitTestResult hitTest = VisualTreeHelper.HitTest(listView, new Point(10, y));
+			DependencyObject depObj = hitTest.VisualHit as DependencyObject;
+			if (depObj == null)
+				return null;
+			DependencyObject current = depObj;
+			while (current != null && current != listView)
+			{
+				ListBoxItem listBoxItem = current as ListBoxItem;
+				if (listBoxItem != null)
+					return listBoxItem;
+				current = VisualTreeHelper.GetParent(current);
+			}
+			return null;
+		}
+		private static char GetCharFromKey(Key key)
+		{
+			char ch = '\0';
+
+			int virtualKey = KeyInterop.VirtualKeyFromKey(key);
+			byte[] keyboardState = new byte[256];
+			NativeMethods.GetKeyboardState(keyboardState);
+
+			uint scanCode = NativeMethods.MapVirtualKey((uint)virtualKey, NativeMethods.MapType.MAPVK_VK_TO_VSC);
+			StringBuilder stringBuilder = new StringBuilder(2);
+
+			int result = NativeMethods.ToUnicode((uint)virtualKey, scanCode, keyboardState, stringBuilder, stringBuilder.Capacity, 0);
+			switch (result)
+			{
+				case -1:
+					break;
+				case 0:
+					break;
+				case 1:
+				default:
+					ch = stringBuilder[0];
+					break;
+			}
+			return ch;
+		}
+
+		class StackItem
+		{
+			public ProjectItems Items;
+			public int Index;
+			public Project Project;
+		}
+
+		class SolutionTraverser
+		{
+			private SortedDictionary<string, FileItem> SortedFileItems = new SortedDictionary<string, FileItem>();
+			private Dictionary<string, Dictionary<string, FileItem>> CollidingFileItems = new Dictionary<string, Dictionary<string, FileItem>>();
+			private Dictionary<string, string> ModuleRoots = new Dictionary<string, string>();
+
+			private Stack<StackItem> ItemStack = new Stack<StackItem>();
+			private int ProjectIndex;
+			public int HandledItemCount { get; private set; }
+
+			public List<FileItem> Update()
+			{
+				Stack<StackItem> itemStack = ItemStack;
+				var projects = UnrealVSPackage.Instance.DTE.Solution.Projects;
+				var projectCount = projects.Count;
+
+				int traverseCounter = 0;
+
+				while (ProjectIndex < projectCount)
+				{
+					var project = projects.Item(ProjectIndex + 1);
+
+					StackItem stackItem;
+					if (ItemStack.Count > 0)
+						stackItem = ItemStack.Pop();
+					else
+						stackItem = new StackItem() { Items = project.ProjectItems, Project = project };
+
+					while (true)
+					{
+						if (traverseCounter > 5000)
+						{
+							ItemStack.Push(stackItem);
+							return null;
+						}
+
+						if (stackItem.Items == null || stackItem.Index == stackItem.Items.Count)
+						{
+							if (itemStack.Count == 0)
+								break;
+							stackItem = itemStack.Pop();
+							++stackItem.Index;
+							continue;
+						}
+
+						var projectItem = stackItem.Items.Item(stackItem.Index + 1);
+						++traverseCounter;
+
+						if (projectItem.FileCount != 0)
+						{
+							var file = projectItem.FileNames[1];
+							if (file != null)
+							{
+								var name = projectItem.Name;
+								if (name != file)
+									if (!file.EndsWith("\\")) // Skip folders
+									{
+										FileItem fileItem = new FileItem() { Name = name, File = file, Project = stackItem.Project.Name };
+										FileItem existingFileItem;
+										if (SortedFileItems.TryGetValue(name, out existingFileItem))
+										{
+											if (existingFileItem.File != file)
+											{
+												Dictionary<string, FileItem> colList;
+												if (!CollidingFileItems.TryGetValue(name, out colList))
+												{
+													colList = new Dictionary<string, FileItem>();
+													colList.Add(existingFileItem.File, existingFileItem);
+													CollidingFileItems[name] = colList;
+												}
+												if (!colList.ContainsKey(fileItem.File))
+													colList.Add(fileItem.File, fileItem);
+											}
+										}
+										else
+											SortedFileItems.Add(name, fileItem);
+
+										if (file.EndsWith(".build.cs", StringComparison.OrdinalIgnoreCase))
+										{
+											var lastBackslash = file.LastIndexOf('\\');
+											var moduleDir = file.Substring(0, lastBackslash);
+											var moduleName = file.Substring(lastBackslash + 1, file.Length - lastBackslash - 10);
+											if (!ModuleRoots.ContainsKey(moduleDir))
+												ModuleRoots.Add(moduleDir, moduleName);
+											//else if (ModuleRoots[moduleDir] != moduleName)
+											//    Logging.WriteLine(file);
+										}
+
+										++HandledItemCount;
+									}
+							}
+						}
+
+						if (projectItem.SubProject != null)
+						{
+							itemStack.Push(stackItem);
+							stackItem = new StackItem() { Items = projectItem.SubProject.ProjectItems, Project = projectItem.SubProject };
+							continue;
+						}
+
+						if (projectItem.ProjectItems != null)
+						{
+							itemStack.Push(stackItem);
+							stackItem = new StackItem() { Items = projectItem.ProjectItems, Project = stackItem.Project };
+							continue;
+						}
+
+						++stackItem.Index;
+					}
+
+					++ProjectIndex;
+				}
+
+				// We need to rename all collisions
+				foreach (var kv in CollidingFileItems)
+				{
+					SortedFileItems.Remove(kv.Key);
+					var usedParents = new Dictionary<string, int>();
+
+					foreach (var itemKv in kv.Value)
+					{
+						var item = itemKv.Value;
+						string parent = item.Project;
+						var path = item.File;
+						int lastBackslash = path.LastIndexOf('\\');
+						while (lastBackslash != -1)
+						{
+							path = path.Substring(0, lastBackslash);
+							string moduleName;
+							if (ModuleRoots.TryGetValue(path, out moduleName))
+							{
+								parent = moduleName;
+								break;
+							}
+							lastBackslash = path.LastIndexOf('\\');
+						}
+
+						string sortName = item.Name + $" ({parent})";
+
+						int counter = 0;
+						if (usedParents.TryGetValue(parent, out counter))
+						{
+							if (counter == 0)
+							{
+								var alreadyAddedItem = SortedFileItems[sortName];
+								SortedFileItems.Remove(sortName);
+								alreadyAddedItem.Name = item.Name + $" ({parent} 0)";
+								SortedFileItems.Add(alreadyAddedItem.Name, alreadyAddedItem);
+							}
+							++counter;
+							usedParents[parent] = counter;
+							item.Name = item.Name + $" ({parent} {counter})";
+							SortedFileItems.Add(item.Name, item);
+						}
+						else
+						{
+							usedParents.Add(parent, 0);
+							item.Name = sortName;
+							SortedFileItems.Add(item.Name, item);
+						}
+					}
+				}
+
+				return new List<FileItem>(SortedFileItems.Values);
+			}
+		}
+	}
 }
\ No newline at end of file
diff --git a/Engine/Source/Programs/UnrealVS/UnrealVS.Shared/P4Commands.cs b/Engine/Source/Programs/UnrealVS/UnrealVS.Shared/P4Commands.cs
--- a/Engine/Source/Programs/UnrealVS/UnrealVS.Shared/P4Commands.cs
+++ b/Engine/Source/Programs/UnrealVS/UnrealVS.Shared/P4Commands.cs
@@ -12,11 +12,11 @@ using System.ComponentModel.Design;
 using System.Diagnostics;
 using System.IO;
 using System.Linq;
+using System.Reflection.Metadata;
 using System.Text;
 using System.Text.RegularExpressions;
 using System.Threading;
 using System.Threading.Tasks;
-using System.Windows;
 using Process = System.Diagnostics.Process;
 
 namespace UnrealVS
@@ -43,12 +43,11 @@ namespace UnrealVS
 		private const int P4TimelapseButtonID = 0x1458;
 		private const int P4RevisionGraphButtonID = 0x1459;
 		private const int P4FileHistoryButtonID = 0x1460;
-
 		private OleMenuCommand SubMenuCommand;
 
 		//private System.Diagnostics.Process ChildProcess;
 
-		private List<System.Diagnostics.Process> ChildProcessList = new List<System.Diagnostics.Process>();
+		private List<Process> ChildProcessList = new List<Process>();
 
 		private IVsOutputWindowPane P4OutputPane;
 		private string P4WorkingDirectory;
@@ -67,22 +66,25 @@ namespace UnrealVS
 		private string Client = "";
 		private string Stream = "";
 		private string UserInfoComplete = "";
-
-		private List<CommandEvents> EventsForce = new List<CommandEvents>();
-
 		private readonly object CheckoutQueueLock = new object();
 		private readonly HashSet<string> CheckoutQueueFiles = new HashSet<string>(StringComparer.OrdinalIgnoreCase);
 		private JoinableTask CheckoutTask;
 
 		private class P4Command
 		{
-			public MenuCommand ButtonCommand;
-			public CommandID CommandID;
-			public P4Command(int ButtonID, EventHandler ButtonHandler)
-			{
-				CommandID = new CommandID(GuidList.UnrealVSCmdSet, ButtonID);
-				ButtonCommand = new MenuCommand(new EventHandler(ButtonHandler), CommandID);
+			public readonly OleMenuCommand ButtonCommand;
 
+			public P4Command(int ButtonID, EventHandler ButtonHandler, Func<EnvDTE.Document, bool> QueryStatusHandler = null)
+			{
+				var cmd = new OleMenuCommand(new EventHandler(ButtonHandler), new CommandID(GuidList.UnrealVSCmdSet, ButtonID));
+				cmd.BeforeQueryStatus += (s, e) =>
+				{
+					ThreadHelper.ThrowIfNotOnUIThread();
+					DTE DTE = UnrealVSPackage.Instance.DTE;
+					((MenuCommand)s).Enabled = DTE.ActiveDocument != null && (QueryStatusHandler == null || QueryStatusHandler(DTE.ActiveDocument));
+				};
+					
+				ButtonCommand = cmd;
 				UnrealVSPackage.Instance.MenuCommandService.AddCommand(ButtonCommand);
 			}
 
@@ -92,7 +94,7 @@ namespace UnrealVS
 			}
 		}
 
-		private List<P4Command> P4CommandsList = new List<P4Command>();
+		private readonly List<P4Command> P4CommandsList = new List<P4Command>(16);
 
 		private InterceptSave Interceptor;
 
@@ -163,18 +165,23 @@ namespace UnrealVS
 			// add commands
 			P4CommandsList.Add(new P4Command(P4CheckoutButtonID, P4CheckoutButtonHandler));
 			P4CommandsList.Add(new P4Command(P4AnnotateButtonID, P4AnnotateButtonHandler));
-			P4CommandsList.Add(new P4Command(P4TimelapseButtonID, P4TimelapseHandler));
-			P4CommandsList.Add(new P4Command(P4IntegrationAwareTimelapseButtonID, P4IntegrationAwareTimeLapseHandler));
-			P4CommandsList.Add(new P4Command(P4RevisionGraphButtonID, P4RevisionGraphHandler));
-			P4CommandsList.Add(new P4Command(P4DiffinVSButtonID, P4DiffinVSHandler));
+			P4CommandsList.Add(new P4Command(P4DiffinVSButtonID, P4DiffHandler));
 			P4CommandsList.Add(new P4Command(P4GetLast10ChangesID, P4GetLast10ChangesHandler));
-			P4CommandsList.Add(new P4Command(P4FileHistoryButtonID, P4FileHistoryHandler));
-			P4CommandsList.Add(new P4Command(P4ShowFileinP4VID, P4ShowFileInP4VHandler));
 			P4CommandsList.Add(new P4Command(P4FastReconcileCodeFilesID, P4FastReconcileCodeFiles));
 
 			if (P4VCCmd.Length > 1)
 			{
-				P4CommandsList.Add(new P4Command(P4ViewSelectedCLButtonID, P4ViewSelectedCLButtonHandler));
+				P4CommandsList.Add(new P4Command(P4IntegrationAwareTimelapseButtonID, P4IntegrationAwareTimeLapseHandler));
+				P4CommandsList.Add(new P4Command(P4ShowFileinP4VID, P4ShowFileInP4VHandler));
+				P4CommandsList.Add(new P4Command(P4FileHistoryButtonID, P4FileHistoryHandler));
+				P4CommandsList.Add(new P4Command(P4RevisionGraphButtonID, P4RevisionGraphHandler));
+				P4CommandsList.Add(new P4Command(P4TimelapseButtonID, P4TimelapseHandler));
+				P4CommandsList.Add(new P4Command(P4ViewSelectedCLButtonID, P4ViewSelectedCLButtonHandler, (Document) =>
+				{
+					ThreadHelper.ThrowIfNotOnUIThread();
+					TextSelection TextSel = (TextSelection)Document.Selection;
+					return !string.IsNullOrEmpty(TextSel.Text) && int.TryParse(TextSel.Text, out _);
+				}));
 			}
 
 			// add sub menu for commands
@@ -226,36 +233,6 @@ namespace UnrealVS
 			bPullWorkingDirectorFromP4 = true;
 		}
 
-		void RegisterCallbackHandler(string CommandName, _dispCommandEvents_BeforeExecuteEventHandler Callback)
-		{
-			ThreadHelper.ThrowIfNotOnUIThread();
-
-			// Find the command from the passed in name
-			DTE DTE = UnrealVSPackage.Instance.DTE;
-			CommandEvents Event = null;
-			{
-				// Probably should move this out to a function.
-				try
-				{
-					Command Command = DTE.Commands.Item(CommandName, -1);
-					if (Command != null)
-					{
-						Event = DTE.Events.get_CommandEvents(Command.Guid, Command.ID);
-					}
-				}
-				catch
-				{
-
-				}
-			}
-
-			if (Event != null)
-			{
-				Event.BeforeExecute += Callback;
-				EventsForce.Add(Event); // forces a reference
-			}
-		}
-
 		private void UpdateMenuOptions()
 		{
 			ThreadHelper.ThrowIfNotOnUIThread();
@@ -428,23 +405,10 @@ namespace UnrealVS
 
 			TextSelection TextSel = (TextSelection)DTE.ActiveDocument.Selection;
 
-			int ChangeList = -1;
-
-			try
-			{
-				ChangeList = Int32.Parse(TextSel.Text);
-			}
-			catch
-			{
-				ChangeList = -2;
-			}
-
-
-			if (ChangeList > 0)
+			if(Int32.TryParse(TextSel.Text, out int ChangeList) && ChangeList > 0)
 			{
 				TryP4VCCommand($"Change {ChangeList}");
 			}
-
 		}
 
 		private void P4TimelapseHandler(object sender, EventArgs args)
@@ -530,7 +494,7 @@ namespace UnrealVS
 			}
 		}
 
-		private void P4DiffinVSHandler(object Sender, EventArgs Args)
+		private void P4DiffHandler(object Sender, EventArgs Args)
 		{
 			ThreadHelper.ThrowIfNotOnUIThread();
 
@@ -545,6 +509,12 @@ namespace UnrealVS
 				return;
 			}
 
+			if(UnrealVSPackage.Instance.OptionsPage.UseP4VDiff && P4VCCmd.Length > 1)
+			{
+				TryP4VCCommand($"diffhave \"{DTE.ActiveDocument.FullName}\"");
+				return;
+			}
+
 			ChangeDiffSetting();
 
 			string CaptureP4Output = "";
@@ -669,8 +639,21 @@ namespace UnrealVS
 				return;
 			}
 
-			TryP4VCommand($"-s \"{DTE.ActiveDocument.FullName}\"");
+			if(UserInfoComplete.Length == 0)
+			{
+				SetUserInfoStrings();
+			}
+
+			if (Username.Length > 0 && Port.Length > 0 && Client.Length > 0)
+			{
+				TryP4VCommand($"-p {Port} -u {Username} -c {Client} -s \"{DTE.ActiveDocument.FullName}\"");
+			}
+			else
+			{
+				TryP4VCommand($"-s \"{DTE.ActiveDocument.FullName}\"");
+			}
 		}
+
 		public void OpenForEdit(string FileName, bool CheckOptionsFlag = false)
 		{
 			ThreadHelper.ThrowIfNotOnUIThread();
diff --git a/Engine/Source/Programs/UnrealVS/UnrealVS.Shared/UnrealVSOptions.cs b/Engine/Source/Programs/UnrealVS/UnrealVS.Shared/UnrealVSOptions.cs
--- a/Engine/Source/Programs/UnrealVS/UnrealVS.Shared/UnrealVSOptions.cs
+++ b/Engine/Source/Programs/UnrealVS/UnrealVS.Shared/UnrealVSOptions.cs
@@ -27,14 +27,19 @@ namespace UnrealVS
 		[Description("Unreal VS will override built in diff settings to the ideal for code, does not alter P4")]
 		public bool AllowUnrealVSOverrideDiffSettings { get; set; }
 
+		[Category("Unreal.P4")]
+		[DisplayName("Use P4V Diff")]
+		[Description("Uses P4V to diff files")]
+		public bool UseP4VDiff { get; set; }
+
 		[Category("Unreal.P4")]
 		[DisplayName("Allow perforce operations")]
-		[Description("Uses p4 ini / environment settings to call P4 functionlality, use 'RunUAT P4WriteConfig' to initialize those settings")]
+		[Description("Uses p4 ini / environment settings to call P4 functionality, use 'RunUAT P4WriteConfig' to initialize those settings")]
 		public bool AllowUnrealVSP4 { get; set; }
 
 		[Category("Unreal.P4")]
 		[DisplayName("Allow Async Checkout")]
-		[Description("Switches checkout (on edit or manual) to an async method, note if the p4 operation fails the local file will still be writeable")]
+		[Description("Switches checkout (on edit or manual) to an async method, note if the p4 operation fails the local file will still be writable")]
 		public bool AllowAsyncP4Checkout { get; set; }
 
 		[Category("Unreal.P4")]
diff --git a/Engine/Source/Programs/UnrealVS/UnrealVS.Shared/UnrealVSPackage.cs b/Engine/Source/Programs/UnrealVS/UnrealVS.Shared/UnrealVSPackage.cs
--- a/Engine/Source/Programs/UnrealVS/UnrealVS.Shared/UnrealVSPackage.cs
+++ b/Engine/Source/Programs/UnrealVS/UnrealVS.Shared/UnrealVSPackage.cs
@@ -75,7 +75,7 @@ namespace UnrealVS
 	{
 		/** Constants */
 
-		private const string VersionString = "v1.79";
+		private const string VersionString = "v1.80";
 		private const string UnrealSolutionFileNamePrefix = "UE";
 		private const string ExtensionName = "UnrealVS";
 		private const string CommandLineOptionKey = ExtensionName + "CommandLineMRU";
diff --git a/Engine/Source/Programs/UnrealVS/UpdateInstructions.txt b/Engine/Source/Programs/UnrealVS/UpdateInstructions.txt
--- a/Engine/Source/Programs/UnrealVS/UpdateInstructions.txt
+++ b/Engine/Source/Programs/UnrealVS/UpdateInstructions.txt
@@ -1,8 +1,11 @@
 1. Make your changes
-2. Update version number in both VSIX manifest files and both AssemblyInfo files.
+2. Update version number in both VSIX manifest files and AssemblyInfo file. 
+	* AssemblyInfo is shared between projects
+	* Resources are also shared
+	* UnrealVS.vsct is also shared
 3. update ReleaseNotes.txt with the changes
 4. Checkout the VSIX in Engine\Extras\UnrealVS\ files and rebuild 2019 IN VS2019, 2022 in VS2022
 5. Update the atomfeed.xml in Engine\Extras\UnrealVS\<version>\atomfeed.xml
-	* <updated> with new timespamp
+	* <updated> with new timestamp
 	* <Version> with new version number
 6. Checkin all the above in 1 CL
