File History no longer hangs forever.

Timelapse now passes in the -l flag to set the line we're interested in.

OutputString is now called as OutputStringThreadSafe, as recommended by documentation.

diff --git a/Engine/Source/Programs/UnrealVS/UnrealVS.Shared/P4Commands.cs b/Engine/Source/Programs/UnrealVS/UnrealVS.Shared/P4Commands.cs
--- a/Engine/Source/Programs/UnrealVS/UnrealVS.Shared/P4Commands.cs
+++ b/Engine/Source/Programs/UnrealVS/UnrealVS.Shared/P4Commands.cs
@@ -156,7 +156,7 @@ namespace UnrealVS
 				else
 				{
 					P4OutputPane.Activate();
-					P4OutputPane.OutputString($"1>------ P4VC not found, {P4VCCmd} or {P4VCCmdBat}{Environment.NewLine}");
+					P4OutputPane.OutputStringThreadSafe($"1>------ P4VC not found, {P4VCCmd} or {P4VCCmdBat}{Environment.NewLine}");
 					P4VCCmd = "";
 				}
 
@@ -200,7 +200,7 @@ namespace UnrealVS
 			string CheckoutQueueFile = GetCheckoutQueueFileName();
 			if (CheckoutQueueFile != null)
 			{
-				P4OutputPane.OutputString($"UnrealVS started. Using checkout queue: {CheckoutQueueFile}" + Environment.NewLine);
+				P4OutputPane.OutputStringThreadSafe($"UnrealVS started. Using checkout queue: {CheckoutQueueFile}" + Environment.NewLine);
 				PulseCheckoutQueue(CheckoutQueueFile);
 			}
 		}
@@ -216,7 +216,7 @@ namespace UnrealVS
 			string CheckoutQueueFile = GetCheckoutQueueFileName();
 			if (CheckoutQueueFile != null)
 			{
-				P4OutputPane.OutputString($"Solution opened. Using checkout queue: {CheckoutQueueFile}" + Environment.NewLine);
+				P4OutputPane.OutputStringThreadSafe($"Solution opened. Using checkout queue: {CheckoutQueueFile}" + Environment.NewLine);
 				PulseCheckoutQueue(CheckoutQueueFile);
 			}
 		}
@@ -285,7 +285,7 @@ namespace UnrealVS
 				Logging.WriteLine("P4Checkout called without an active document");
 
 				P4OutputPane.Activate();
-				P4OutputPane.OutputString($"1>------ P4Checkout called without an active document{Environment.NewLine}");
+				P4OutputPane.OutputStringThreadSafe($"1>------ P4Checkout called without an active document{Environment.NewLine}");
 
 				return;
 			}
@@ -307,7 +307,7 @@ namespace UnrealVS
 			{
 				Logging.WriteLine("P4Annotate called without an active document");
 
-				P4OutputPane.OutputString($"1>------ P4Annotate called without an active document{Environment.NewLine}");
+				P4OutputPane.OutputStringThreadSafe($"1>------ P4Annotate called without an active document{Environment.NewLine}");
 
 				return;
 			}
@@ -320,7 +320,7 @@ namespace UnrealVS
 
 			if (!bResult || CaptureP4StdErr.Length > 0)
 			{
-				P4OutputPane.OutputString($"1>------ P4Annotate call failed");
+				P4OutputPane.OutputStringThreadSafe($"1>------ P4Annotate call failed");
 				return;
 			}
 
@@ -420,11 +420,24 @@ namespace UnrealVS
 			if (DTE.ActiveDocument == null)
 			{
 				Logging.WriteLine("P4Timelapse called without an active document");
-				P4OutputPane.OutputString($"1>------ P4Timelapse called without an active document{Environment.NewLine}");
+				P4OutputPane.OutputStringThreadSafe($"1>------ P4Timelapse called without an active document{Environment.NewLine}");
 				return;
 			}
 
-			string Command = $"timelapse \"{DTE.ActiveDocument.FullName}\"";
+			string Command = "timelapse ";
+
+			TextDocument CurrentDocument = (TextDocument) DTE.ActiveDocument.Object("TextDocument");
+			if (CurrentDocument != null)
+			{
+				TextSelection TextSel = (TextSelection) DTE.ActiveDocument.Selection;
+				if (TextSel != null)
+				{
+					Command += $"-l {TextSel.CurrentLine} ";
+				}
+			}
+
+			Command += $"\"{DTE.ActiveDocument.FullName}\"";
+
 			TryP4VCCommand(Command, true);
 		}
 
@@ -437,7 +450,7 @@ namespace UnrealVS
 			if (DTE.ActiveDocument == null)
 			{
 				Logging.WriteLine("P4Timelapse called without an active document");
-				P4OutputPane.OutputString($"1>------ P4RevisionGraph called without an active document{Environment.NewLine}");
+				P4OutputPane.OutputStringThreadSafe($"1>------ P4RevisionGraph called without an active document{Environment.NewLine}");
 				return;
 			}
 
@@ -447,6 +460,11 @@ namespace UnrealVS
 
 		private void P4FileHistoryHandler(object sender, EventArgs args)
 		{
+			if (P4VCCmd.Length == 0)
+			{
+				return;
+			}
+
 			ThreadHelper.ThrowIfNotOnUIThread();
 
 			DTE DTE = UnrealVSPackage.Instance.DTE;
@@ -454,12 +472,25 @@ namespace UnrealVS
 			if (DTE.ActiveDocument == null)
 			{
 				Logging.WriteLine("P4FileHistory called without an active document");
-				P4OutputPane.OutputString($"1>------ P4FileHistory called without an active document{Environment.NewLine}");
+				P4OutputPane.OutputStringThreadSafe($"1>------ P4FileHistory called without an active document{Environment.NewLine}");
 				return;
 			}
 
 			string Command = $"history \"{DTE.ActiveDocument.FullName}\"";
-			TryP4VCCommand(Command, true);
+			P4OutputPane.OutputStringThreadSafe($"1>------ Running P4FileHistory with args: {Command}{Environment.NewLine}");
+
+			Process ChildProcess = new Process()
+			{
+				StartInfo = new ProcessStartInfo()
+				{
+					FileName = P4VCCmd,
+					Arguments = Command,
+					WorkingDirectory = P4WorkingDirectory,
+					UseShellExecute = false,
+					CreateNoWindow = true
+				}
+			};
+			ChildProcess.Start();
 		}
 
 		private void P4IntegrationAwareTimeLapseHandler(object Sender, EventArgs Args)
@@ -472,7 +503,7 @@ namespace UnrealVS
 			{
 				Logging.WriteLine("P4IntegrationAwareTimeLapse called without an active document");
 
-				P4OutputPane.OutputString($"1>------ P4IntegrationAwareTimeLapse called without an active document{Environment.NewLine}");
+				P4OutputPane.OutputStringThreadSafe($"1>------ P4IntegrationAwareTimeLapse called without an active document{Environment.NewLine}");
 
 				return;
 			}
@@ -504,7 +535,7 @@ namespace UnrealVS
 			{
 				Logging.WriteLine("P4DiffInVSHandler called without an active document");
 
-				P4OutputPane.OutputString($"1>------ P4DiffInVSHandler called without an active document{Environment.NewLine}");
+				P4OutputPane.OutputStringThreadSafe($"1>------ P4DiffInVSHandler called without an active document{Environment.NewLine}");
 
 				return;
 			}
@@ -619,10 +650,10 @@ namespace UnrealVS
 			}
 			else
 			{
-				P4OutputPane.OutputString($"PREVIEW{Environment.NewLine}");
+				P4OutputPane.OutputStringThreadSafe($"PREVIEW{Environment.NewLine}");
 			}
 
-			P4OutputPane.OutputString($"Running Async Reconcile on : {ReconcileDepotPaths}{Environment.NewLine}");
+			P4OutputPane.OutputStringThreadSafe($"Running Async Reconcile on : {ReconcileDepotPaths}{Environment.NewLine}");
 			P4OutputPane.Activate();
 
 			_ = TryP4CommandAsync($"reconcile -e -m {Preview} {ReconcileDepotPaths}");
@@ -661,7 +692,7 @@ namespace UnrealVS
 			// if the callee requested it, honor the enabling of autocheckout option - return if its off
 			if (CheckOptionsFlag && !UnrealVSPackage.Instance.OptionsPage.AllowUnrealVSCheckoutOnEdit)
 			{
-				//P4OutputPane.OutputString($"autocheckout disabled: {FileName}");
+				//P4OutputPane.OutputStringThreadSafe($"autocheckout disabled: {FileName}");
 				return;
 			}
 
@@ -669,7 +700,7 @@ namespace UnrealVS
 			if (!File.Exists(FileName) ||
  				(CheckOptionsFlag && !File.GetAttributes(FileName).HasFlag(FileAttributes.ReadOnly)))
 			{
-				P4OutputPane.OutputString($"already writeable: {FileName}" + Environment.NewLine);
+				P4OutputPane.OutputStringThreadSafe($"already writeable: {FileName}" + Environment.NewLine);
 				return;
 			}
 
@@ -678,7 +709,7 @@ namespace UnrealVS
 				string queueFile = GetCheckoutQueueFileName();
 				if (queueFile == null)
 				{
-					P4OutputPane.OutputString("Unable to get path to checkout queue file. No solution loaded?");
+					P4OutputPane.OutputStringThreadSafe("Unable to get path to checkout queue file. No solution loaded?" + Environment.NewLine);
 					return;
 				}
 
@@ -692,12 +723,12 @@ namespace UnrealVS
 					}
 					catch (Exception ex)
 					{
-						P4OutputPane.OutputString($"Unable to update {queueFile} ({ex.Message})" + Environment.NewLine);
+						P4OutputPane.OutputStringThreadSafe($"Unable to update {queueFile} ({ex.Message})" + Environment.NewLine);
 						return;
 					}
 				}
 
-				//P4OutputPane.OutputString($"Marking file as writeable: {FileName}");
+				//P4OutputPane.OutputStringThreadSafe($"Marking file as writeable: {FileName}");
 				// Mark the file as writeable
 				FileAttributes NewAttributes = File.GetAttributes(FileName) & ~FileAttributes.ReadOnly;
 				File.SetAttributes(FileName, NewAttributes);
@@ -837,7 +868,7 @@ namespace UnrealVS
 
 			if (!bSuccess)
 			{
-				P4OutputPane.OutputString($"attempt to pull P4CONFIG info failed{Environment.NewLine}");
+				P4OutputPane.OutputStringThreadSafe($"attempt to pull P4CONFIG info failed{Environment.NewLine}");
 			}
 
 			return WorkingDirectory;
@@ -867,7 +898,7 @@ namespace UnrealVS
 
 			UserInfoComplete = string.Format(" -p {0} -u {1} -c {2} ", Port, Username, Client);
 
-			P4OutputPane.OutputString("GetUserInfoStringFull : " + UserInfoComplete + Environment.NewLine);
+			P4OutputPane.OutputStringThreadSafe("GetUserInfoStringFull : " + UserInfoComplete + Environment.NewLine);
 		}
 
 		private void PullWorkingDirectory(bool bPullfromP4Settings)
@@ -899,11 +930,11 @@ namespace UnrealVS
 						P4WorkingDirectory = NewWorkingDirectory;
 						bPullWorkingDirectorFromP4 = false;
 
-						P4OutputPane.OutputString($"P4WorkingDirectory set to '{P4WorkingDirectory}'{Environment.NewLine}");
+						P4OutputPane.OutputStringThreadSafe($"P4WorkingDirectory set to '{P4WorkingDirectory}'{Environment.NewLine}");
 					}
 					else
 					{
-						P4OutputPane.OutputString($"P4WorkingDirectory set failed {Environment.NewLine}");
+						P4OutputPane.OutputStringThreadSafe($"P4WorkingDirectory set failed {Environment.NewLine}");
 					}
 				}
 			}
@@ -962,7 +993,7 @@ namespace UnrealVS
 					P4OutputPane.Activate();
 				}
 
-				P4OutputPane.OutputString($"Running \"{CmdPath}\" {CommandLine}{Environment.NewLine}");
+				P4OutputPane.OutputStringThreadSafe($"Running \"{CmdPath}\" {CommandLine}{Environment.NewLine}");
 			});
 
 			StringBuilder StdOutSB = new StringBuilder();
@@ -1019,15 +1050,15 @@ namespace UnrealVS
 				await ThreadHelper.JoinableTaskFactory.SwitchToMainThreadAsync();
 				if (StdOut.Length > 0 && bOutputStdOut)
 				{
-					P4OutputPane.OutputString(StdOut);
+					P4OutputPane.OutputStringThreadSafe(StdOut);
 				}
 				if (StdErr.Length > 0)
 				{
-					P4OutputPane.OutputString(StdErr);
+					P4OutputPane.OutputStringThreadSafe(StdErr);
 				}
 
 				TimeSpan Duration = DateTime.Now - Start;
-				P4OutputPane.OutputString($"complete {Duration.TotalSeconds.ToString()} {Environment.NewLine}");
+				P4OutputPane.OutputStringThreadSafe($"complete {Duration.TotalSeconds.ToString()} {Environment.NewLine}");
 			});
 
 			bool Result = (ChildProcess.ExitCode == 0 && StdErr.Length == 0);
