UnrealVS:
- Added Revert command
- Checkout now also Opens for Add if the file is not already added.

diff --git a/Engine/Source/Programs/UnrealVS/UnrealVS.2019/Properties/AssemblyInfo.cs b/Engine/Source/Programs/UnrealVS/UnrealVS.2019/Properties/AssemblyInfo.cs
--- a/Engine/Source/Programs/UnrealVS/UnrealVS.2019/Properties/AssemblyInfo.cs
+++ b/Engine/Source/Programs/UnrealVS/UnrealVS.2019/Properties/AssemblyInfo.cs
@@ -31,8 +31,8 @@ using System.Runtime.InteropServices;
 // You can specify all the values or you can default the Revision and Build Numbers 
 // by using the '*' as shown below:
 
-[assembly: AssemblyVersion("1.80.0.0")]
-[assembly: AssemblyFileVersion("1.80.0.0")]
+[assembly: AssemblyVersion("1.82.0.0")]
+[assembly: AssemblyFileVersion("1.82.0.0")]
 
 
 
diff --git a/Engine/Source/Programs/UnrealVS/UnrealVS.2022/UnrealVS.2022.csproj b/Engine/Source/Programs/UnrealVS/UnrealVS.2022/UnrealVS.2022.csproj
--- a/Engine/Source/Programs/UnrealVS/UnrealVS.2022/UnrealVS.2022.csproj
+++ b/Engine/Source/Programs/UnrealVS/UnrealVS.2022/UnrealVS.2022.csproj
@@ -104,11 +104,15 @@
     <Resource Include="..\UnrealVS.2019\Resources\Icon.ico">
       <Link>Resources\Icon.ico</Link>
     </Resource>
-    <Resource Include="..\UnrealVS.2019\Resources\Icon.png">
+    <Content Include="..\UnrealVS.2019\Resources\Icon.png">
       <Link>Resources\Icon.png</Link>
-    </Resource>
+      <CopyToOutputDirectory>Always</CopyToOutputDirectory>
+      <IncludeInVSIX>true</IncludeInVSIX>
+    </Content>
     <Resource Include="..\UnrealVS.2019\Resources\Preview.png">
       <Link>Resources\Preview.png</Link>
+      <IncludeInVSIX>true</IncludeInVSIX>
+      <CopyToOutputDirectory>Always</CopyToOutputDirectory>
     </Resource>
     <Resource Include="..\UnrealVS.2019\Resources\RefreshProjects.png">
       <Link>Resources\RefreshProjects.png</Link>
diff --git a/Engine/Source/Programs/UnrealVS/UnrealVS.Shared/P4Commands.cs b/Engine/Source/Programs/UnrealVS/UnrealVS.Shared/P4Commands.cs
--- a/Engine/Source/Programs/UnrealVS/UnrealVS.Shared/P4Commands.cs
+++ b/Engine/Source/Programs/UnrealVS/UnrealVS.Shared/P4Commands.cs
@@ -43,6 +43,8 @@ namespace UnrealVS
 		private const int P4TimelapseButtonID = 0x1458;
 		private const int P4RevisionGraphButtonID = 0x1459;
 		private const int P4FileHistoryButtonID = 0x1460;
+		private const int P4RevertButtonID = 0x1461;
+
 		private OleMenuCommand SubMenuCommand;
 
 		//private System.Diagnostics.Process ChildProcess;
@@ -164,6 +166,7 @@ namespace UnrealVS
 
 			// add commands
 			P4CommandsList.Add(new P4Command(P4CheckoutButtonID, P4CheckoutButtonHandler));
+			P4CommandsList.Add(new P4Command(P4RevertButtonID, P4RevertButtonHandler));
 			P4CommandsList.Add(new P4Command(P4AnnotateButtonID, P4AnnotateButtonHandler));
 			P4CommandsList.Add(new P4Command(P4DiffinVSButtonID, P4DiffHandler));
 			P4CommandsList.Add(new P4Command(P4GetLast10ChangesID, P4GetLast10ChangesHandler));
@@ -273,6 +276,29 @@ namespace UnrealVS
 			}
 		}
 
+		private void P4RevertButtonHandler(object Sender, EventArgs Args)
+		{
+			ThreadHelper.ThrowIfNotOnUIThread();
+
+			DTE DTE = UnrealVSPackage.Instance.DTE;
+
+			// Check we've got a file open
+			if (DTE.ActiveDocument == null)
+			{
+				Logging.WriteLine("P4Checkout called without an active document");
+
+				P4OutputPane.Activate();
+				P4OutputPane.OutputStringThreadSafe($"1>------ P4Revert called without an active document{Environment.NewLine}");
+
+				return;
+			}
+
+			if (VSConstants.MessageBoxResult.IDOK == (VSConstants.MessageBoxResult) VsShellUtilities.ShowMessageBox(UnrealVSPackage.Instance, string.Empty, $"Are you sure you want to revert {DTE.ActiveDocument.Name}?", OLEMSGICON.OLEMSGICON_QUERY, OLEMSGBUTTON.OLEMSGBUTTON_OKCANCEL, OLEMSGDEFBUTTON.OLEMSGDEFBUTTON_FIRST))
+			{
+				TryP4Command($"revert {DTE.ActiveDocument.FullName}", out _, out _);
+			}
+		}
+
 		private void P4CheckoutButtonHandler(object Sender, EventArgs Args)
 		{
 			ThreadHelper.ThrowIfNotOnUIThread();
@@ -738,8 +764,16 @@ namespace UnrealVS
 			}
 			else
 			{
+				string StdOut, StdErr;
+
 				// sync edit command - will lock the UI but be safer (no risk of a locally writeable file that is not checked out)
-				TryP4Command($"edit \"{FileName}\"", out _, out _);
+				bool Success = TryP4Command($"edit \"{FileName}\"", out StdOut, out StdErr);
+
+				if (!Success && StdErr.Contains("file(s) not on client."))
+				{
+					P4OutputPane.OutputStringThreadSafe($"{FileName} has not been added yet, adding now.{Environment.NewLine}");
+					TryP4Command($"add \"{FileName}\"", out _, out _);
+				}
 			}
 		}
 
@@ -789,33 +823,57 @@ namespace UnrealVS
 					bool Pause = false;
 					foreach (string Line in Lines)
 					{
-						string TrimLine = Line.Trim();
-						if (TrimLine.Length > 0)
+						string FileName = Line.Trim();
+						if (FileName.Length > 0)
 						{
-							bool Success = await TryP4CommandAsync($"edit \"{TrimLine}\"");
-							lock (CheckoutQueueLock)
+							bool Success = false;
+
+							(bool EditSuccess, string StdOut, string StdErr) = await TryP4CommandExAsync(P4Exe, $"edit \"{FileName}\"");
+							if (EditSuccess)
 							{
-								if (Success)
+								Success = true;
+
+								Logging.WriteLine($"Performed async checkout of {FileName}.");
+								lock (CheckoutQueueLock)
 								{
-									Logging.WriteLine($"Performed async checkout of {TrimLine}.");
+									string[] NewLines = File.ReadAllLines(QueueFile);
+									File.WriteAllLines(QueueFile, NewLines.Where(x => !x.Equals(Line, StringComparison.Ordinal)));
+								}
+							}
+
+							if (!EditSuccess && StdErr.Contains("file(s) not on client."))
+							{
+								await ThreadHelper.JoinableTaskFactory.RunAsync(async () =>
+								{
+									await ThreadHelper.JoinableTaskFactory.SwitchToMainThreadAsync();
+									P4OutputPane.OutputStringThreadSafe($"{FileName} has not been added yet, adding now.{Environment.NewLine}");
+								});
+								(bool AddSuccess, string _, string _) = await TryP4CommandExAsync(P4Exe, $"add \"{FileName}\"");
+
+								if (AddSuccess)
+								{
+									Success = true;
+
+									Logging.WriteLine($"Performed async add of {FileName}.");
 									lock (CheckoutQueueLock)
 									{
 										string[] NewLines = File.ReadAllLines(QueueFile);
 										File.WriteAllLines(QueueFile, NewLines.Where(x => !x.Equals(Line, StringComparison.Ordinal)));
 									}
 								}
-								else
-								{
-									// Set the UpdateCheckoutQueue flag again, to force another loop
-									Logging.WriteLine($"Unable to check out file {TrimLine}. Will retry.");
-									CheckoutQueueFiles.Add(QueueFile);
-									Pause = true;
-								}
+							}
+
+							if (!Success)
+							{
+								// Re-add to checkout queue, to force another loop
+								Logging.WriteLine($"Unable to check out file {FileName}. Will retry.");
+								CheckoutQueueFiles.Add(QueueFile);
+								Pause = true;
 							}
 						}
 					}
 
-					// Check if we're done
+					// Check if we've looped over everything and should pause
 					if (Pause)
 					{
 						await Task.Delay(TimeSpan.FromSeconds(10.0));
diff --git a/Engine/Source/Programs/UnrealVS/UnrealVS.Shared/UnrealVSOptions.cs b/Engine/Source/Programs/UnrealVS/UnrealVS.Shared/UnrealVSOptions.cs
--- a/Engine/Source/Programs/UnrealVS/UnrealVS.Shared/UnrealVSOptions.cs
+++ b/Engine/Source/Programs/UnrealVS/UnrealVS.Shared/UnrealVSOptions.cs
@@ -18,12 +18,12 @@ namespace UnrealVS
 		public bool HideNonGameStartupProjects { get; set; }
 
 		[Category("Unreal.P4")]
-		[DisplayName("Enable auto checkout on save")]
-		[Description("Uses p4 ini / environment settings to automatically checkout files on save, use 'RunUAT P4WriteConfig' to initialize those settings")]
+		[DisplayName("Enable Auto Checkout on Save")]
+		[Description("Uses P4 ini / environment settings to automatically checkout files on save, use 'RunUAT P4WriteConfig' to initialize those settings")]
 		public bool AllowUnrealVSCheckoutOnEdit { get; set; }
 
 		[Category("Unreal.P4")]
-		[DisplayName("Override VS compare options")]
+		[DisplayName("Override VS Compare Options")]
 		[Description("Unreal VS will override built in diff settings to the ideal for code, does not alter P4")]
 		public bool AllowUnrealVSOverrideDiffSettings { get; set; }
 
@@ -33,8 +33,8 @@ namespace UnrealVS
 		public bool UseP4VDiff { get; set; }
 
 		[Category("Unreal.P4")]
-		[DisplayName("Allow perforce operations")]
-		[Description("Uses p4 ini / environment settings to call P4 functionality, use 'RunUAT P4WriteConfig' to initialize those settings")]
+		[DisplayName("Allow Perforce operations")]
+		[Description("Uses P4 ini / environment settings to call P4 functionality, use 'RunUAT P4WriteConfig' to initialize those settings")]
 		public bool AllowUnrealVSP4 { get; set; }
 
 		[Category("Unreal.P4")]
@@ -43,8 +43,8 @@ namespace UnrealVS
 		public bool AllowAsyncP4Checkout { get; set; }
 
 		[Category("Unreal.P4")]
-		[DisplayName("Force P4 Output window to active")]
-		[Description("When set, normal p4 operations will force the output window to the surface")]
+		[DisplayName("Force P4 Output Window to Active")]
+		[Description("When set, normal P4 operations will force the output window to the surface")]
 		public bool ForceOutputWindow { get; set; }
 
 		[Category("Unreal.P4")]
@@ -58,7 +58,7 @@ namespace UnrealVS
 		public string ReconcileExtensions { get; set; } = "c*;h*;ini;uproject;uplugin";
 
 		[Category("General")]
-		[DisplayName("Prefix UE5 solution name with folder")]
+		[DisplayName("Prefix UE5 Solution Name with Folder")]
 		[Description("Will include name of folder that UE5.sln is located in to Visual Studio window title.")]
 		public bool IncludeFolderInUE5SolutionName { get; set; } = true;
 
