[UnrealVS]
* FileBrowser - Added so ctrl-shift-I is pasting selected file as include in to current opened document
* CompileSingleFile - Added command CompileAndProfileSingleFile that run compile score around the compile (need version [REDACTED_PROFILE_URL]/CompileScore for now, will do pull request up to main branch)
* Fixed some code based on warnings compiling UnrealVS

diff --git a/Engine/Source/Programs/UnrealVS/UnrealVS.2019/Properties/AssemblyInfo.cs b/Engine/Source/Programs/UnrealVS/UnrealVS.2019/Properties/AssemblyInfo.cs
--- a/Engine/Source/Programs/UnrealVS/UnrealVS.2019/Properties/AssemblyInfo.cs
+++ b/Engine/Source/Programs/UnrealVS/UnrealVS.2019/Properties/AssemblyInfo.cs
@@ -31,8 +31,8 @@ using System.Runtime.InteropServices;
 // You can specify all the values or you can default the Revision and Build Numbers 
 // by using the '*' as shown below:
 
-[assembly: AssemblyVersion("1.82.0.0")]
-[assembly: AssemblyFileVersion("1.82.0.0")]
+[assembly: AssemblyVersion("1.84.0.0")]
+[assembly: AssemblyFileVersion("1.84.0.0")]
 
 
 
diff --git a/Engine/Source/Programs/UnrealVS/UnrealVS.Shared/BatchBuilderToolControl.xaml.cs b/Engine/Source/Programs/UnrealVS/UnrealVS.Shared/BatchBuilderToolControl.xaml.cs
--- a/Engine/Source/Programs/UnrealVS/UnrealVS.Shared/BatchBuilderToolControl.xaml.cs
+++ b/Engine/Source/Programs/UnrealVS/UnrealVS.Shared/BatchBuilderToolControl.xaml.cs
@@ -495,7 +495,7 @@ namespace UnrealVS
 				// Clear and activate the output pane.
 				Pane.Clear();
 
-				Pane.OutputString(Text);
+				Pane.OutputStringThreadSafe(Text);
 
 				// @todo: Activating doesn't seem to really bring the pane to front like we would expect it to.
 				Pane.Activate();
diff --git a/Engine/Source/Programs/UnrealVS/UnrealVS.Shared/CompileSingleFile.cs b/Engine/Source/Programs/UnrealVS/UnrealVS.Shared/CompileSingleFile.cs
--- a/Engine/Source/Programs/UnrealVS/UnrealVS.Shared/CompileSingleFile.cs
+++ b/Engine/Source/Programs/UnrealVS/UnrealVS.Shared/CompileSingleFile.cs
@@ -19,10 +19,11 @@ namespace UnrealVS
 		private const int CompileSingleFileButtonID = 0x1075;
 		private const int PreprocessSingleFileButtonID = 0x1076;
 		private const int CompileSingleModuleButtonID = 0x1077;
+		private const int CompileAndProfileSingleFileButtonID = 0x1078;
 		private const int UBTSubMenuID = 0x3103;
 		private string	  FileToCompileOriginalExt = "";
 
-		static readonly HashSet<string> ValidExtensions = new HashSet<string>(StringComparer.OrdinalIgnoreCase) { ".c", ".cc", ".cpp", ".cxx" };
+		static readonly HashSet<string> ValidExtensions = new HashSet<string>(StringComparer.OrdinalIgnoreCase) { ".c", ".cc", ".cpp", ".h", ".cxx", ".ispc" };
 
 		System.Diagnostics.Process ChildProcess;
 		private OleMenuCommand SubMenuCommand;
@@ -43,6 +44,10 @@ namespace UnrealVS
 			MenuCommand CompileSingleModuleButtonCommand = new MenuCommand(new EventHandler(CompileSingleFileButtonHandler), CommandID3);
 			UnrealVSPackage.Instance.MenuCommandService.AddCommand(CompileSingleModuleButtonCommand);
 
+			CommandID CommandID4 = new CommandID(GuidList.UnrealVSCmdSet, CompileAndProfileSingleFileButtonID);
+			MenuCommand CompileAndProfileSingleFileButtonCommand = new MenuCommand(new EventHandler(CompileSingleFileButtonHandler), CommandID4);
+			UnrealVSPackage.Instance.MenuCommandService.AddCommand(CompileAndProfileSingleFileButtonCommand);
+
 			// add sub menu for UBT commands
 			SubMenuCommand = new OleMenuCommand(null, new CommandID(GuidList.UnrealVSCmdSet, UBTSubMenuID));
 			UnrealVSPackage.Instance.MenuCommandService.AddCommand(SubMenuCommand);
@@ -81,9 +86,10 @@ namespace UnrealVS
 			MenuCommand SenderSubMenuCommand = (MenuCommand)Sender;
 
 			bool bIsFile = SenderSubMenuCommand.CommandID.ID != CompileSingleModuleButtonID;
-			bool PreprocessOnly = SenderSubMenuCommand.CommandID.ID == PreprocessSingleFileButtonID;
+			bool bPreprocessOnly = SenderSubMenuCommand.CommandID.ID == PreprocessSingleFileButtonID;
+			bool bProfile = SenderSubMenuCommand.CommandID.ID == CompileAndProfileSingleFileButtonID;
 
-			if (!TryCompileSingleFileOrModule(bIsFile, PreprocessOnly))
+			if (!TryCompileSingleFileOrModule(bIsFile, bPreprocessOnly, bProfile))
 			{
 				DTE DTE = UnrealVSPackage.Instance.DTE;
 				DTE.ExecuteCommand("Build.Compile");
@@ -120,7 +126,7 @@ namespace UnrealVS
 			return FindModuleForFile(directory, rootDirectory);
 		}
 
-		bool TryCompileSingleFileOrModule(bool bIsFile, bool bPreProcessOnly)
+		bool TryCompileSingleFileOrModule(bool bIsFile, bool bPreProcessOnly, bool bProfile)
 		{
 			ThreadHelper.ThrowIfNotOnUIThread();
 			DTE DTE = UnrealVSPackage.Instance.DTE;
@@ -143,7 +149,7 @@ namespace UnrealVS
 				if (MessageBox.Show("Cancel current compile?", "Compile in progress", MessageBoxButtons.YesNo) == DialogResult.Yes)
 				{
 					KillChildProcess();
-					BuildOutputPane.OutputString($"1>  Build cancelled.{Environment.NewLine}");
+					BuildOutputPane.OutputStringThreadSafe($"1>  Build cancelled.{Environment.NewLine}");
 				}
 				return true;
 			}
@@ -188,7 +194,7 @@ namespace UnrealVS
 			Microsoft.VisualStudio.VCProjectEngine.VCNMakeTool ActiveNMakeTool = (ActiveVCConfiguration.Tools as Microsoft.VisualStudio.VCProjectEngine.IVCCollection).Item("VCNMakeTool") as Microsoft.VisualStudio.VCProjectEngine.VCNMakeTool;
 			if (ActiveNMakeTool == null)
 			{
-				MessageBox.Show($"No NMakeTool set for Project {VCStartupProject.Name} set for single-file compile.", "NMakeTool not set", MessageBoxButtons.OK);
+				MessageBox.Show($"No NMakeTool set for Project {VCStartupProject.Name} set for single-file compile.", "UnrealVS - NMakeTool not set", MessageBoxButtons.OK);
 				return false;
 			}
 
@@ -239,8 +245,8 @@ namespace UnrealVS
 			// Set up the output pane
 			BuildOutputPane.Activate();
 			BuildOutputPane.Clear();
-			BuildOutputPane.OutputString($"1>------ Build started: Project: {StartupProject.Name}, Configuration: {ActiveConfiguration.ConfigurationName} {ActiveConfiguration.PlatformName} ------{Environment.NewLine}");
-			BuildOutputPane.OutputString($"1>  Compiling {CompilingText}{Environment.NewLine}");
+			BuildOutputPane.OutputStringThreadSafe($"1>------ Build started: Project: {StartupProject.Name}, Configuration: {ActiveConfiguration.ConfigurationName} {ActiveConfiguration.PlatformName} ------{Environment.NewLine}");
+			BuildOutputPane.OutputStringThreadSafe($"1>  Compiling {CompilingText}{Environment.NewLine}");
 
 			// Set up event handlers 
 			DTE.Events.BuildEvents.OnBuildBegin += BuildEvents_OnBuildBegin;
@@ -251,7 +257,7 @@ namespace UnrealVS
 			{ 
 				if (Args.Data != null) 
 				{ 
-					BuildOutputPane.OutputString($"1>  {Args.Data}{Environment.NewLine}"); 
+					BuildOutputPane.OutputStringThreadSafe($"1>  {Args.Data}{Environment.NewLine}"); 
 					if (Args.Data.Contains("PreProcessPath:"))
 					{
 						PPLine = Args.Data;
@@ -259,6 +265,23 @@ namespace UnrealVS
 				} 
 			}
 
+			if (bProfile)
+			{
+				var Commands = DTE.Commands.Cast<Command>();
+				#pragma warning disable VSTHRD010
+				Command Command = Commands.FirstOrDefault((C) => C.Name == "CompileScore.StartTrace");
+				#pragma warning restore VSTHRD010
+				if (Command == null || !Command.IsAvailable)
+				{
+					MessageBox.Show($"CompileSingleFileAndProfile requires a CompileScore visual studio extension that has CompileScore.StartTrace/CompileScore.StopTrace installed", "UnrealVS - Missing CompileScore extension ", MessageBoxButtons.OK);
+					return true;
+				}
+				Object CustomIn = null;
+				Object CustomOut = null;
+				DTE.Commands.Raise(Command.Guid, Command.ID, ref CustomIn, ref CustomOut);
+				//DTE.ExecuteCommand("CompileScore.StartTrace");
+			}
+
 			string SolutionDir = Path.GetDirectoryName(UnrealVSPackage.Instance.SolutionFilepath).TrimEnd(Path.DirectorySeparatorChar) + Path.DirectorySeparatorChar;
 			// Get the build command line and escape any environment variables that we use
 			string BuildCommandLine = ActiveNMakeTool.BuildCommandLine;
@@ -280,12 +303,12 @@ namespace UnrealVS
 			ChildProcess.StartInfo.CreateNoWindow = true;
 			ChildProcess.OutputDataReceived += OutputHandler;
 			ChildProcess.ErrorDataReceived += OutputHandler;
-			if (bPreProcessOnly)
+			if (bPreProcessOnly || bProfile)
 			{ 
 				// add an event handler to respond to the exit of the preprocess request
 				// and open the generated file if it exists.
 				ChildProcess.EnableRaisingEvents = true;
-				ChildProcess.Exited += new EventHandler((s, e) => PreprocessExitHandler(PPLine));
+				ChildProcess.Exited += new EventHandler((s, e) => PreprocessExitHandler(PPLine, bProfile));
 			}
 			
 			ChildProcess.Start();
@@ -295,8 +318,17 @@ namespace UnrealVS
 			return true;
 		}
 
-		private void PreprocessExitHandler(string PPLine)
+		private void PreprocessExitHandler(string PPLine, bool bIsProfiling)
 		{
+			if (bIsProfiling)
+			{
+				ThreadHelper.JoinableTaskFactory.Run(async () =>
+				{
+					await ThreadHelper.JoinableTaskFactory.SwitchToMainThreadAsync();
+					UnrealVSPackage.Instance.DTE.ExecuteCommand("CompileScore.StopTrace");
+				});
+			}
+
 			// not all compile actions support pre-process - check it exists
 			if (PPLine.Contains("PreProcessPath:"))
 			{
@@ -317,7 +349,7 @@ namespace UnrealVS
 			IVsOutputWindowPane BuildOutputPane = UnrealVSPackage.Instance.GetOutputPane();
 			DTE DTE = UnrealVSPackage.Instance.DTE;
 
-			BuildOutputPane.OutputString($"1>  PPFullPath: {PPFullPath}{Environment.NewLine}");
+			BuildOutputPane.OutputStringThreadSafe($"1>  PPFullPath: {PPFullPath}{Environment.NewLine}");
 
 			if (File.Exists(PPFullPath))
 			{
diff --git a/Engine/Source/Programs/UnrealVS/UnrealVS.Shared/FileBrowserWindowControl.xaml b/Engine/Source/Programs/UnrealVS/UnrealVS.Shared/FileBrowserWindowControl.xaml
--- a/Engine/Source/Programs/UnrealVS/UnrealVS.Shared/FileBrowserWindowControl.xaml
+++ b/Engine/Source/Programs/UnrealVS/UnrealVS.Shared/FileBrowserWindowControl.xaml
@@ -169,6 +169,7 @@
                         <TextBlock Text="F1"/>
                         <TextBlock Text="F5"/>
                         <TextBlock Text="Ctrl-I"/>
+                        <TextBlock Text="Ctrl-Shift-I"/>
                     </StackPanel>
                     <StackPanel Grid.Column="1" Margin="6">
                         <TextBlock Text="- Close FileBrowser"/>
@@ -181,6 +182,7 @@
                         <TextBlock Text="- This help"/>
                         <TextBlock Text="- Refresh file lists with current visual studio solution"/>
                         <TextBlock Text="- Create a #include &quot;file&quot; statement to clipboard and close file browser"/>
+                        <TextBlock Text="- Create a #include &quot;file&quot; statement and paste it into active document"/>
                     </StackPanel>
                 </Grid>
                 <TextBlock Text=""/>
diff --git a/Engine/Source/Programs/UnrealVS/UnrealVS.Shared/FileBrowserWindowControl.xaml.cs b/Engine/Source/Programs/UnrealVS/UnrealVS.Shared/FileBrowserWindowControl.xaml.cs
--- a/Engine/Source/Programs/UnrealVS/UnrealVS.Shared/FileBrowserWindowControl.xaml.cs
+++ b/Engine/Source/Programs/UnrealVS/UnrealVS.Shared/FileBrowserWindowControl.xaml.cs
@@ -6,9 +6,9 @@ using Microsoft.VisualStudio.Shell.Interop;
 using System;
 using System.Collections.Generic;
 using System.IO;
+using System.Linq;
 using System.Text;
 using System.Text.Json;
-using System.Text.Json.Serialization;
 using System.Windows;
 using System.Windows.Controls;
 using System.Windows.Data;
@@ -36,7 +36,7 @@ namespace UnrealVS
 
 		String UnrealVsFileName;
 
-		List<FileItem> AllFileItems;
+		FileItem[] AllFileItems;
 		List<FileItem> BookmarkedFileItems;
 		List<FileItem> RecentFileItems;
 
@@ -47,7 +47,7 @@ namespace UnrealVS
 			Window = window;
 			this.InitializeComponent();
 
-			AllFileItems = new List<FileItem>();
+			AllFileItems = Array.Empty<FileItem>();
 			BookmarkedFileItems = new List<FileItem>();
 			RecentFileItems = new List<FileItem>();
 			IsVisibleChanged += FileBrowserWindowControl_IsVisibleChanged;
@@ -102,6 +102,7 @@ namespace UnrealVS
 					return;
 
 				FilterEditBox.Text = Path.GetFileName(text);
+				FilterEditBox.SelectAll();
 			}
 		}
 
@@ -239,8 +240,139 @@ namespace UnrealVS
 			if ((Keyboard.Modifiers & ModifierKeys.Control) != 0)
 			{
 				if (e.Key == Key.I)
-					if (CreateIncludePath())
-						HidePanel();
+				{
+					if (!CreateIncludePath())
+					{
+						return;
+					}
+					
+					HidePanel();
+
+					if ((Keyboard.Modifiers & ModifierKeys.Shift) == 0)
+					{
+						return;
+					}
+
+					// Let's add directly in to active document
+					var activeDoc = UnrealVSPackage.Instance.DTE.ActiveDocument;
+					var doc = (TextDocument)(activeDoc.Object("TextDocument"));
+					if (doc == null)
+						return;
+
+					bool addSpace = false;
+					var includePath = GetIncludePath((FileItem)(ActiveFilesListBox.SelectedItems[0]));
+					string pasteString = $"#include \"{includePath}\"";
+
+					bool skipFirstInclude = activeDoc.FullName.EndsWith(".cpp");
+					int addBeforeLine = -1;
+					var lastIncludeLine = -1;
+
+					bool isInComment = false;
+					var p = doc.StartPoint.CreateEditPoint();
+					var lastLine = doc.EndPoint.Line;
+					for (int lineIndex = 1; lineIndex < lastLine; ++lineIndex)
+					{
+						var str = p.GetLines(lineIndex, lineIndex + 1).Trim();
+						if (str.Length == 0)
+						{
+							continue;
+						}
+
+						if (isInComment)
+						{
+							if (str.IndexOf("*/") != -1)
+								isInComment = false;
+							continue;
+						}
+
+						if (str.StartsWith("//"))
+						{
+							continue;
+						}
+
+						if (str.StartsWith("/*"))
+						{
+							isInComment = true;
+							continue;
+						}
+
+						if (str.StartsWith("#include"))
+						{
+							if (str == pasteString)
+							{
+								doc.Selection.MoveTo(lineIndex, 0);
+								doc.Selection.SelectLine();
+								p.TryToShow();
+								return;
+							}
+
+							bool isFirst = lastIncludeLine == -1;
+							int prevLastIncludeLine = lastIncludeLine;
+							lastIncludeLine = lineIndex;
+							if (skipFirstInclude && isFirst)
+							{
+								continue;
+							}
+
+							var span = str.AsSpan(9);
+							var firstQuote = span.IndexOf('"');
+							if (firstQuote != -1)
+							{
+								span = span.Slice(firstQuote + 1);
+								var secondQuote = span.IndexOf('"');
+								if (secondQuote != -1)
+								{
+									span = span.Slice(0, secondQuote).Trim();
+									bool isGeneratedInclude = span.IndexOf(".generated.".AsSpan()) != -1;
+									if (includePath.AsSpan().CompareTo(span, StringComparison.Ordinal) < 0 || isGeneratedInclude)
+									{
+										if (isGeneratedInclude && prevLastIncludeLine != -1)
+											lastIncludeLine = prevLastIncludeLine;
+										else
+											addBeforeLine = lineIndex;
+										break;
+									}
+								}
+							}
+							else if (span.IndexOf("UE_INLINE_GENERATED_CPP".AsSpan()) != -1)
+							{
+								lastIncludeLine = prevLastIncludeLine;
+								break;
+							}
+							continue;
+						}
+							
+						if (str.StartsWith("#pragma once"))
+						{
+							lastIncludeLine = lineIndex + 1;
+							continue;
+						}
+
+						lastIncludeLine = lineIndex - 1;
+						addSpace = true;
+						break;  // Unknown code
+					}
+
+					if (addBeforeLine == -1 && lastIncludeLine != -1)
+					{
+						addBeforeLine = lastIncludeLine + 1;
+					}
+
+					if (addBeforeLine == -1)
+					{
+						addBeforeLine = lastLine;
+					}
+
+					p.LineDown(addBeforeLine - 1);
+					p.Insert($"{pasteString}\r\n");
+					if (addSpace)
+					{
+						p.Insert("\r\n");
+					}
+					doc.Selection.MoveTo(addBeforeLine, 0);
+					doc.Selection.SelectLine();
+					p.TryToShow();
+				}
 				return;
 			}
 			if (e.Key == Key.Left)
@@ -287,7 +419,7 @@ namespace UnrealVS
 				}
 				if (!BookmarksModified)
 					return;
-				RefreshListView(BookmarkedFilesListBox, BookmarkedFileItems);
+				RefreshListView(BookmarkedFilesListBox, BookmarkedFileItems.ToArray());
 				//if (AllFilesListBox.SelectedIndex < AllFilesListBox.Items.Count - 1)
 				//    FocusListViewItem(AllFilesListBox, AllFilesListBox.SelectedIndex + 1);
 				SaveSolutionSettings();
@@ -312,7 +444,7 @@ namespace UnrealVS
 				}
 				if (BookmarksModified)
 				{
-					RefreshListView(BookmarkedFilesListBox, BookmarkedFileItems);
+					RefreshListView(BookmarkedFilesListBox, BookmarkedFileItems.ToArray());
 					if (index == BookmarkedFilesListBox.Items.Count)
 					{
 						if (index == 0)
@@ -402,21 +534,27 @@ namespace UnrealVS
 			return toOpen.Count != 0;
 		}
 
+		private string GetIncludePath(FileItem item)
+		{
+                if (item == null || item.Name == "")
+				return null;
+                string includePath = item.File;
+                includePath = includePath.Replace("\\", "/");
+                string includePathLwr = includePath.ToLower();
+			var includeStartIndex = Math.Max(Math.Max(includePathLwr.LastIndexOf("/public/"), includePathLwr.LastIndexOf("/private/")), includePathLwr.LastIndexOf("/classes/"));
+                if (includeStartIndex != -1)
+                    includePath = includePath.Substring(includePath.IndexOf('/', includeStartIndex + 1) + 1);
+			return includePath;
+		}
+
 		private bool CreateIncludePath()
 		{
 			string includeStrings = "";
 			foreach (var i in ActiveFilesListBox.SelectedItems)
 			{
-				var item = (FileItem)i;
-				if (item == null || item.Name == "")
-					continue;
-				string includePath = item.File;
-				includePath = includePath.Replace("\\", "/");
-				string includePathLwr = includePath.ToLower();
-				var includeStartIndex = Math.Max(Math.Max(includePathLwr.LastIndexOf("/public/"), includePathLwr.LastIndexOf("/private/")), includePathLwr.LastIndexOf("/classes/"));
-				if (includeStartIndex != -1)
-					includePath = includePath.Substring(includePath.IndexOf('/', includeStartIndex + 1) + 1);
-				includeStrings += $"#include \"{includePath}\"\r\n";
+				var includePath = GetIncludePath((FileItem)i);
+				if (includePath != null)
+                includeStrings += $"#include \"{includePath}\"\r\n";
 			}
 			if (includeStrings == "")
 				return false;
@@ -481,8 +619,8 @@ namespace UnrealVS
 		private void RefreshListViews()
 		{
 			RefreshListView(AllFilesListBox, AllFileItems);
-			RefreshListView(BookmarkedFilesListBox, BookmarkedFileItems);
-			RefreshListView(RecentFilesListBox, RecentFileItems);
+			RefreshListView(BookmarkedFilesListBox, BookmarkedFileItems.ToArray());
+			RefreshListView(RecentFilesListBox, RecentFileItems.ToArray());
 		}
 
 		struct Score
@@ -491,7 +629,7 @@ namespace UnrealVS
 			public int Index;
 		}
 
-		private void RefreshListView(ListBox listView, List<FileItem> source)
+		private void RefreshListView(ListBox listView, FileItem[] source)
 		{
 			var fileItems = (List<FileItem>)listView.ItemsSource;
 			fileItems.Clear();
@@ -587,11 +725,11 @@ namespace UnrealVS
 
 
 			if (listView == AllFilesListBox)
-				AllFilesTab.Header = $"All Files ({fileCount}/{source.Count})";
+				AllFilesTab.Header = $"All Files ({fileCount}/{source.Length})";
 			else if (listView == BookmarkedFilesListBox)
-				BookmarkedFilesTab.Header = $"Bookmarked Files ({fileCount}/{source.Count})";
+				BookmarkedFilesTab.Header = $"Bookmarked Files ({fileCount}/{source.Length})";
 			else
-				RecentFilesTab.Header = $"Recent Files ({fileCount}/{source.Count})";
+				RecentFilesTab.Header = $"Recent Files ({fileCount}/{source.Length})";
 		}
 
 		private void RefreshStatusBox()
@@ -645,7 +783,7 @@ namespace UnrealVS
 			if (IsVisible)
 				AsyncRefreshFileList();
 			else
-				AllFileItems.Clear(); // Will trigger AsyncRefreshFileList when visible
+				AllFileItems = Array.Empty<FileItem>(); // Will trigger AsyncRefreshFileList when visible
 		}
 
 		internal void HandleDocumentActivated(Document Document)
@@ -662,7 +800,7 @@ namespace UnrealVS
 			if (RecentFileItems.Count > RecentMaxCount)
 				RecentFileItems.RemoveRange(RecentMaxCount, RecentFileItems.Count - RecentMaxCount);
 
-			RefreshListView(RecentFilesListBox, RecentFileItems);
+			RefreshListView(RecentFilesListBox, RecentFileItems.ToArray());
 		}
 
 		private void AsyncRefreshFileList()
@@ -717,7 +855,7 @@ namespace UnrealVS
 			ClearFilter = true;
 			//FilesListTab.SelectedIndex = 0;
 
-			if (AllFileItems.Count == 0)
+			if (AllFileItems.Length == 0)
 			{
 				RefreshListViews(); // To make sure we get one invisible entry
 				AsyncRefreshFileList();
@@ -847,14 +985,14 @@ namespace UnrealVS
 		class SolutionTraverser
 		{
 			private SortedDictionary<string, FileItem> SortedFileItems = new SortedDictionary<string, FileItem>();
-			private Dictionary<string, Dictionary<string, FileItem>> CollidingFileItems = new Dictionary<string, Dictionary<string, FileItem>>();
+			private Dictionary<string, SortedDictionary<string, FileItem>> CollidingFileItems = new Dictionary<string, SortedDictionary<string, FileItem>>();
 			private Dictionary<string, string> ModuleRoots = new Dictionary<string, string>();
 
 			private Stack<StackItem> ItemStack = new Stack<StackItem>();
 			private int ProjectIndex;
 			public int HandledItemCount { get; private set; }
 
-			public List<FileItem> Update()
+			public FileItem[] Update()
 			{
 				Stack<StackItem> itemStack = ItemStack;
 				var projects = UnrealVSPackage.Instance.DTE.Solution.Projects;
@@ -907,10 +1045,10 @@ namespace UnrealVS
 										{
 											if (existingFileItem.File != file)
 											{
-												Dictionary<string, FileItem> colList;
+												SortedDictionary<string, FileItem> colList;
 												if (!CollidingFileItems.TryGetValue(name, out colList))
 												{
-													colList = new Dictionary<string, FileItem>();
+													colList = new SortedDictionary<string, FileItem>();
 													colList.Add(existingFileItem.File, existingFileItem);
 													CollidingFileItems[name] = colList;
 												}
@@ -1007,7 +1145,29 @@ namespace UnrealVS
 					}
 				}
 
-				return new List<FileItem>(SortedFileItems.Values);
+				var result = SortedFileItems.Values.ToArray();
+
+				// Sort so header is before source file.. since most of the time you are opening header files.
+				/*
+                for (int i=0, e=result.Length; i!=e; ++i)
+                {
+                    var name = result[i].Name;
+					if (name.EndsWith(".h") && i > 0)
+                    {
+						var prev = result[i - 1];
+						var prevName = prev.Name;
+                        if (name.Length + 2 == prevName.Length)
+                        {
+                            if (name.AsSpan(0, name.Length - 1).SequenceEqual(prevName.AsSpan(0, prevName.Length - 3))) // swap h and cpp
+                            {
+                                result[i - 1] = result[i];
+                                result[i] = prev;
+							}
+                        }
+                    }
+                }
+                */
+				return result;
 			}
 		}
 	}
diff --git a/Engine/Source/Programs/UnrealVS/UnrealVS.Shared/GenerateProjectFiles.cs b/Engine/Source/Programs/UnrealVS/UnrealVS.Shared/GenerateProjectFiles.cs
--- a/Engine/Source/Programs/UnrealVS/UnrealVS.Shared/GenerateProjectFiles.cs
+++ b/Engine/Source/Programs/UnrealVS/UnrealVS.Shared/GenerateProjectFiles.cs
@@ -138,7 +138,7 @@ namespace UnrealVS
 				var Pane = UnrealVSPackage.Instance.GetOutputPane();
 				if (Pane != null)
 				{
-					Pane.OutputString(Args.Data + "\n");
+					Pane.OutputStringThreadSafe(Args.Data + "\n");
 				}
 			});
 			AsyncTask.Join();
