[UnrealVS] Fix parsing P4CONFIG working directory

diff --git a/Engine/Source/Programs/UnrealVS/UnrealVS.Shared/P4Commands.cs b/Engine/Source/Programs/UnrealVS/UnrealVS.Shared/P4Commands.cs
--- a/Engine/Source/Programs/UnrealVS/UnrealVS.Shared/P4Commands.cs
+++ b/Engine/Source/Programs/UnrealVS/UnrealVS.Shared/P4Commands.cs
@@ -915,7 +915,7 @@ namespace UnrealVS
 			string[] lines = CaptureP4Output.Split('\n');
 			foreach(string Line in lines)
 			{
-				System.Text.RegularExpressions.Match Match = Regex.Match(Line, @"^P4CONFIG\s*=.*'([^']+)'\)$");
+				System.Text.RegularExpressions.Match Match = Regex.Match(Line, @"^P4CONFIG\s*=.*'([^']+)'\s*\)$");
 				if (Match.Success)
 				{
 					WorkingDirectory = Path.GetDirectoryName(Match.Groups[1].Value);
