UnrealVS:  Be more aggressive by trying more P4 Settings and also more relaxed about what constitutes a valid line when finding the P4WorkingDirectory

diff --git a/Engine/Extras/UnrealVS/VS2022/AtomFeed.xml b/Engine/Extras/UnrealVS/VS2022/AtomFeed.xml
--- a/Engine/Extras/UnrealVS/VS2022/AtomFeed.xml
+++ b/Engine/Extras/UnrealVS/VS2022/AtomFeed.xml
@@ -8,7 +8,7 @@
     <title type="text">UnrealVS</title>
     <summary type="text">Unreal Engine extension for Visual Studio</summary>
     <published>2022-10-21T10:01:57-07:00</published>
-    <updated>2023-01-31T13:00:00-00:00</updated>
+    <updated>2023-02-24T13:00:00-00:00</updated>
     <author>
       <name>Epic Games, Inc.</name>
     </author>
@@ -19,7 +19,7 @@
     <content type="application/octet-stream" src="UnrealVS.vsix" />
     <Vsix xmlns:xsd="http://www.w3.org/2001/XMLSchema" xmlns:xsi="http://www.w3.org/2001/XMLSchema-instance" xmlns="http://schemas.microsoft.com/developer/vsx-syndication-schema/2010">
       <Id>ddbf523f-7eb6-4887-bd51-85a714ff87eb</Id>
-      <Version>1.83</Version>
+      <Version>1.84.1</Version>
       <Rating xsi:nil="true" />
       <RatingCount xsi:nil="true" />
       <DownloadCount xsi:nil="true" />
diff --git a/Engine/Source/Programs/UnrealVS/ReleaseNotes.txt b/Engine/Source/Programs/UnrealVS/ReleaseNotes.txt
--- a/Engine/Source/Programs/UnrealVS/ReleaseNotes.txt
+++ b/Engine/Source/Programs/UnrealVS/ReleaseNotes.txt
@@ -1,3 +1,6 @@
+1.84.1
+* P4 Commands now are more aggressive about finding your working directory
+
 1.83
 * Open for Edit now also Adds if the file is not already added
 * added Revert command
diff --git a/Engine/Source/Programs/UnrealVS/Unreal.TestAdapter/Unreal.TestAdapter.csproj b/Engine/Source/Programs/UnrealVS/Unreal.TestAdapter/Unreal.TestAdapter.csproj
--- a/Engine/Source/Programs/UnrealVS/Unreal.TestAdapter/Unreal.TestAdapter.csproj
+++ b/Engine/Source/Programs/UnrealVS/Unreal.TestAdapter/Unreal.TestAdapter.csproj
@@ -32,16 +32,17 @@
   </PropertyGroup>
   <ItemGroup>
     <Reference Include="Microsoft.TestPlatform.CoreUtilities, Version=15.0.0.0, Culture=neutral, PublicKeyToken=b03f5f7f11d50a3a, processorArchitecture=MSIL">
-      <HintPath>..\packages\Microsoft.TestPlatform.ObjectModel.16.10.0\lib\net451\Microsoft.TestPlatform.CoreUtilities.dll</HintPath>
+      <HintPath>..\packages\Microsoft.TestPlatform.ObjectModel.17.5.0\lib\net462\Microsoft.TestPlatform.CoreUtilities.dll</HintPath>
     </Reference>
     <Reference Include="Microsoft.TestPlatform.PlatformAbstractions, Version=15.0.0.0, Culture=neutral, PublicKeyToken=b03f5f7f11d50a3a, processorArchitecture=MSIL">
-      <HintPath>..\packages\Microsoft.TestPlatform.ObjectModel.16.10.0\lib\net451\Microsoft.TestPlatform.PlatformAbstractions.dll</HintPath>
+      <HintPath>..\packages\Microsoft.TestPlatform.ObjectModel.17.5.0\lib\net462\Microsoft.TestPlatform.PlatformAbstractions.dll</HintPath>
     </Reference>
     <Reference Include="Microsoft.VisualStudio.TestPlatform.ObjectModel, Version=15.0.0.0, Culture=neutral, PublicKeyToken=b03f5f7f11d50a3a, processorArchitecture=MSIL">
-      <HintPath>..\packages\Microsoft.TestPlatform.ObjectModel.16.10.0\lib\net451\Microsoft.VisualStudio.TestPlatform.ObjectModel.dll</HintPath>
+      <HintPath>..\packages\Microsoft.TestPlatform.ObjectModel.17.5.0\lib\net462\Microsoft.VisualStudio.TestPlatform.ObjectModel.dll</HintPath>
     </Reference>
     <Reference Include="NuGet.Frameworks, Version=5.11.0.10, Culture=neutral, PublicKeyToken=31bf3856ad364e35, processorArchitecture=MSIL">
-      <HintPath>..\packages\NuGet.Frameworks.5.11.0\lib\net472\NuGet.Frameworks.dll</HintPath>
+      <HintPath>..\packages\NuGet.Frameworks.6.5.0\lib\net472\NuGet.Frameworks.dll</HintPath>
+      <Private>True</Private>
     </Reference>
     <Reference Include="System" />
     <Reference Include="System.Collections.Immutable, Version=1.2.3.0, Culture=neutral, PublicKeyToken=b03f5f7f11d50a3a, processorArchitecture=MSIL">
diff --git a/Engine/Source/Programs/UnrealVS/Unreal.TestAdapter/packages.config b/Engine/Source/Programs/UnrealVS/Unreal.TestAdapter/packages.config
--- a/Engine/Source/Programs/UnrealVS/Unreal.TestAdapter/packages.config
+++ b/Engine/Source/Programs/UnrealVS/Unreal.TestAdapter/packages.config
@@ -1,8 +1,8 @@
 ﻿<?xml version="1.0" encoding="utf-8"?>
 <packages>
-  <package id="Microsoft.TestPlatform" version="17.4.1" targetFramework="net472" />
-  <package id="Microsoft.TestPlatform.ObjectModel" version="17.4.1" targetFramework="net472" />
-  <package id="NuGet.Frameworks" version="5.11.3" targetFramework="net472" />
+  <package id="Microsoft.TestPlatform" version="17.5.0" targetFramework="net472" />
+  <package id="Microsoft.TestPlatform.ObjectModel" version="17.5.0" targetFramework="net472" />
+  <package id="NuGet.Frameworks" version="6.5.0" targetFramework="net472" />
   <package id="System.Collections.Immutable" version="1.5.0" targetFramework="net472" />
   <package id="System.Reflection.Metadata" version="1.6.0" targetFramework="net472" />
 </packages>
\ No newline at end of file
diff --git a/Engine/Source/Programs/UnrealVS/UnrealVS.2019/Properties/AssemblyInfo.cs b/Engine/Source/Programs/UnrealVS/UnrealVS.2019/Properties/AssemblyInfo.cs
--- a/Engine/Source/Programs/UnrealVS/UnrealVS.2019/Properties/AssemblyInfo.cs
+++ b/Engine/Source/Programs/UnrealVS/UnrealVS.2019/Properties/AssemblyInfo.cs
@@ -31,6 +31,6 @@ using System.Runtime.InteropServices;
 // You can specify all the values or you can default the Revision and Build Numbers 
 // by using the '*' as shown below:
 
-[assembly: AssemblyVersion("1.84.0.0")]
-[assembly: AssemblyFileVersion("1.84.0.0")]
+[assembly: AssemblyVersion("1.84.1.0")]
+[assembly: AssemblyFileVersion("1.84.1.0")]
 
diff --git a/Engine/Source/Programs/UnrealVS/UnrealVS.2019/UnrealVS.2019.csproj b/Engine/Source/Programs/UnrealVS/UnrealVS.2019/UnrealVS.2019.csproj
--- a/Engine/Source/Programs/UnrealVS/UnrealVS.2019/UnrealVS.2019.csproj
+++ b/Engine/Source/Programs/UnrealVS/UnrealVS.2019/UnrealVS.2019.csproj
@@ -203,7 +203,7 @@
   <Import Project="$(VSToolsPath)\VSSDK\Microsoft.VsSDK.targets" Condition="'$(VSToolsPath)' != ''" />
   <Import Project="$(MSBuildExtensionsPath32)\Microsoft\VisualStudio\v11.0\VSSDK\Microsoft.VsSDK.targets" Condition="false" />
   <PropertyGroup>
-    <PostBuildEvent>copy "$(TargetDir)$(TargetName).vsix" "$(ProjectDir)..\..\..\..\Extras\UnrealVS\VS2019\$(TargetName).vsix" /Y</PostBuildEvent>
+    <PostBuildEvent>xcopy /F /Y "$(TargetDir)$(TargetName).vsix" "$(ProjectDir)..\..\..\..\Extras\UnrealVS\VS2019\$(TargetName).vsix"</PostBuildEvent>
   </PropertyGroup>
   <PropertyGroup>
     <PreBuildEvent>
diff --git a/Engine/Source/Programs/UnrealVS/UnrealVS.2022/UnrealVS.2022.csproj b/Engine/Source/Programs/UnrealVS/UnrealVS.2022/UnrealVS.2022.csproj
--- a/Engine/Source/Programs/UnrealVS/UnrealVS.2022/UnrealVS.2022.csproj
+++ b/Engine/Source/Programs/UnrealVS/UnrealVS.2022/UnrealVS.2022.csproj
@@ -232,7 +232,7 @@
   <Import Project="$(VSToolsPath)\VSSDK\Microsoft.VsSDK.targets" Condition="'$(VSToolsPath)' != ''" />
   <Import Project="$(MSBuildExtensionsPath32)\Microsoft\VisualStudio\v11.0\VSSDK\Microsoft.VsSDK.targets" Condition="false" />
   <PropertyGroup>
-    <PostBuildEvent>copy "$(TargetDir)$(TargetName).vsix" "$(ProjectDir)..\..\..\..\Extras\UnrealVS\VS2022\$(TargetName).vsix" /Y</PostBuildEvent>
+    <PostBuildEvent>xcopy /F /Y "$(TargetDir)$(TargetName).vsix" "$(ProjectDir)..\..\..\..\Extras\UnrealVS\VS2022\$(TargetName).vsix"</PostBuildEvent>
   </PropertyGroup>
   <PropertyGroup>
     <PreBuildEvent>
diff --git a/Engine/Source/Programs/UnrealVS/UnrealVS.Shared/P4Commands.cs b/Engine/Source/Programs/UnrealVS/UnrealVS.Shared/P4Commands.cs
--- a/Engine/Source/Programs/UnrealVS/UnrealVS.Shared/P4Commands.cs
+++ b/Engine/Source/Programs/UnrealVS/UnrealVS.Shared/P4Commands.cs
@@ -913,20 +913,26 @@ namespace UnrealVS
 			bool bSuccess = false;
 
 			string[] lines = CaptureP4Output.Split('\n');
-			foreach(string Line in lines)
+			foreach (string Line in lines)
 			{
-				System.Text.RegularExpressions.Match Match = Regex.Match(Line, @"^P4CONFIG\s*=.*'([^']+)'\s*\)$");
+				string TrimLine = Line.Trim();
+
+				// Match this in https://regex101.com/ to debug it. Basically it's words=anything(anything'path-to-config-file'anything).
+				System.Text.RegularExpressions.Match Match = Regex.Match(TrimLine, @"^\w+=.*\(.*'(.*)'.*\)$");
 				if (Match.Success)
 				{
 					WorkingDirectory = Path.GetDirectoryName(Match.Groups[1].Value);
-					bSuccess = true;
-					break;
+					if (Directory.Exists(WorkingDirectory))
+					{
+						bSuccess = true;
+						break;
+					}
 				}
 			}
 
 			if (!bSuccess)
 			{
-				P4OutputPane.OutputStringThreadSafe($"attempt to pull P4CONFIG info failed{Environment.NewLine}");
+				P4OutputPane.OutputStringThreadSafe($"Attempt to pull the P4WorkingDirectory from 'p4 set' failed.  Have you run 'RunUAT P4WriteConfig' in your root directory?{Environment.NewLine}");
 			}
 
 			return WorkingDirectory;
diff --git a/Engine/Source/Programs/UnrealVS/UpdateInstructions.txt b/Engine/Source/Programs/UnrealVS/UpdateInstructions.txt
--- a/Engine/Source/Programs/UnrealVS/UpdateInstructions.txt
+++ b/Engine/Source/Programs/UnrealVS/UpdateInstructions.txt
@@ -1,11 +1,15 @@
 1. Make your changes
-2. Update version number in both VSIX manifest files and AssemblyInfo file. 
+2. Update version number in both VSIX manifest files and the AssemblyInfo file. 
 	* AssemblyInfo is shared between projects
 	* Resources are also shared
 	* UnrealVS.vsct is also shared
+	* Manifest file is named source.extension.vsixmanifest and is unique per project
 3. update ReleaseNotes.txt with the changes
-4. Checkout the VSIX in Engine\Extras\UnrealVS\ files and rebuild 2019 IN VS2019, 2022 in VS2022
+4. Checkout the VSIX in Engine\Extras\UnrealVS\ files and rebuild both UnrealVS.2019 and UnrealVS.2022
 5. Update the atomfeed.xml in Engine\Extras\UnrealVS\<version>\atomfeed.xml
 	* <updated> with new timestamp
 	* <Version> with new version number
 6. Checkin all the above in 1 CL
+
+Development Tips:
+	* In Visual Studio: Tools -> Options.  Then Environment/Extensions.  Add your path to AtomFeed.xml as an additional extension gallery.
