fix for UnrealVS TestAdapater not being able to parse test output
changed to output to a file so that the output is not combined with UE_LOG lines

https://p4-swarm.epicgames.net/reviews/24413087

diff --git a/Engine/Extras/UnrealVS/VS2022/AtomFeed.xml b/Engine/Extras/UnrealVS/VS2022/AtomFeed.xml
--- a/Engine/Extras/UnrealVS/VS2022/AtomFeed.xml
+++ b/Engine/Extras/UnrealVS/VS2022/AtomFeed.xml
@@ -8,7 +8,7 @@
     <title type="text">UnrealVS</title>
     <summary type="text">Unreal Engine extension for Visual Studio</summary>
     <published>2022-10-21T10:01:57-07:00</published>
-    <updated>2023-02-24T13:00:00-00:00</updated>
+    <updated>2023-02-24T13:10:00-00:00</updated>
     <author>
       <name>Epic Games, Inc.</name>
     </author>
@@ -19,7 +19,7 @@
     <content type="application/octet-stream" src="UnrealVS.vsix" />
     <Vsix xmlns:xsd="http://www.w3.org/2001/XMLSchema" xmlns:xsi="http://www.w3.org/2001/XMLSchema-instance" xmlns="http://schemas.microsoft.com/developer/vsx-syndication-schema/2010">
       <Id>ddbf523f-7eb6-4887-bd51-85a714ff87eb</Id>
-      <Version>1.84.1</Version>
+      <Version>1.84.2</Version>
       <Rating xsi:nil="true" />
       <RatingCount xsi:nil="true" />
       <DownloadCount xsi:nil="true" />
diff --git a/Engine/Source/Programs/UnrealVS/ReleaseNotes.txt b/Engine/Source/Programs/UnrealVS/ReleaseNotes.txt
--- a/Engine/Source/Programs/UnrealVS/ReleaseNotes.txt
+++ b/Engine/Source/Programs/UnrealVS/ReleaseNotes.txt
@@ -1,3 +1,6 @@
+1.84.2
+* Fix to Unreal.TestAdapater to output test results to a file
+
 1.84.1
 * P4 Commands now are more aggressive about finding your working directory
 
diff --git a/Engine/Source/Programs/UnrealVS/Unreal.TestAdapter/TestExecutor.cs b/Engine/Source/Programs/UnrealVS/Unreal.TestAdapter/TestExecutor.cs
--- a/Engine/Source/Programs/UnrealVS/Unreal.TestAdapter/TestExecutor.cs
+++ b/Engine/Source/Programs/UnrealVS/Unreal.TestAdapter/TestExecutor.cs
@@ -64,17 +64,19 @@ namespace Unreal.TestAdapter
 				if (File.Exists(tempFile))
 				{
 					xml = File.ReadAllText(tempFile);
+					File.Delete(tempFile);
 				}
 			}
 			else
 			{
+				System.Diagnostics.Debugger.Launch();
+				string tempFile = Path.Combine(Path.GetTempPath(), Path.GetTempFileName());
 				using (Process process = new Process())
 				{
 					process.StartInfo.FileName = testCase.Source;
 
 					//wrap in quotes because test names can have spaces
-					process.StartInfo.Arguments = "\"" + testCase.FullyQualifiedName + "\"  --reporter unreal";
-					process.StartInfo.RedirectStandardOutput = true;
+					process.StartInfo.Arguments = "\"" + testCase.FullyQualifiedName + "\"  --reporter unreal -o " + tempFile;
 					process.StartInfo.UseShellExecute = false;
 
 					if (token.IsCancellationRequested)
@@ -84,22 +86,21 @@ namespace Unreal.TestAdapter
 
 					frameworkHandle.RecordStart(testCase);
 					process.Start();
-					xml = process.StandardOutput.ReadToEnd();
-					while (!process.HasExited)
+					process.WaitForExit();
+					if (File.Exists(tempFile))
 					{
-						if (token.IsCancellationRequested)
-						{
-							process.Kill();
-							return;
-						}
+						xml = File.ReadAllText(tempFile);
+						File.Delete(tempFile);
 					}
 				}
 			}
 
-			var testResult = new TestResult(testCase) { Outcome = TestOutcome.NotFound };
-			ParseResult(xml, testResult);
-			frameworkHandle.RecordEnd(testCase, testResult.Outcome);
-			frameworkHandle.RecordResult(testResult);
+			{
+				var testResult = new TestResult(testCase) { Outcome = TestOutcome.NotFound };
+				ParseResult(xml, testResult);
+				frameworkHandle.RecordEnd(testCase, testResult.Outcome);
+				frameworkHandle.RecordResult(testResult);
+			}
 		}
 
 		void ParseResult(string xmlString, TestResult testResult)
diff --git a/Engine/Source/Programs/UnrealVS/Unreal.TestAdapter/Unreal.TestAdapter.csproj b/Engine/Source/Programs/UnrealVS/Unreal.TestAdapter/Unreal.TestAdapter.csproj
--- a/Engine/Source/Programs/UnrealVS/Unreal.TestAdapter/Unreal.TestAdapter.csproj
+++ b/Engine/Source/Programs/UnrealVS/Unreal.TestAdapter/Unreal.TestAdapter.csproj
@@ -31,35 +31,15 @@
     <WarningLevel>4</WarningLevel>
   </PropertyGroup>
   <ItemGroup>
-    <Reference Include="Microsoft.TestPlatform.CoreUtilities, Version=15.0.0.0, Culture=neutral, PublicKeyToken=b03f5f7f11d50a3a, processorArchitecture=MSIL">
-      <HintPath>..\packages\Microsoft.TestPlatform.ObjectModel.17.5.0\lib\net462\Microsoft.TestPlatform.CoreUtilities.dll</HintPath>
-    </Reference>
-    <Reference Include="Microsoft.TestPlatform.PlatformAbstractions, Version=15.0.0.0, Culture=neutral, PublicKeyToken=b03f5f7f11d50a3a, processorArchitecture=MSIL">
-      <HintPath>..\packages\Microsoft.TestPlatform.ObjectModel.17.5.0\lib\net462\Microsoft.TestPlatform.PlatformAbstractions.dll</HintPath>
-    </Reference>
-    <Reference Include="Microsoft.VisualStudio.TestPlatform.ObjectModel, Version=15.0.0.0, Culture=neutral, PublicKeyToken=b03f5f7f11d50a3a, processorArchitecture=MSIL">
-      <HintPath>..\packages\Microsoft.TestPlatform.ObjectModel.17.5.0\lib\net462\Microsoft.VisualStudio.TestPlatform.ObjectModel.dll</HintPath>
-    </Reference>
-    <Reference Include="NuGet.Frameworks, Version=5.11.0.10, Culture=neutral, PublicKeyToken=31bf3856ad364e35, processorArchitecture=MSIL">
-      <HintPath>..\packages\NuGet.Frameworks.6.5.0\lib\net472\NuGet.Frameworks.dll</HintPath>
-      <Private>True</Private>
-    </Reference>
     <Reference Include="System" />
-    <Reference Include="System.Collections.Immutable, Version=1.2.3.0, Culture=neutral, PublicKeyToken=b03f5f7f11d50a3a, processorArchitecture=MSIL">
-      <HintPath>..\packages\System.Collections.Immutable.1.5.0\lib\netstandard2.0\System.Collections.Immutable.dll</HintPath>
-    </Reference>
     <Reference Include="System.Configuration" />
     <Reference Include="System.Core" />
-    <Reference Include="System.Reflection.Metadata, Version=1.4.3.0, Culture=neutral, PublicKeyToken=b03f5f7f11d50a3a, processorArchitecture=MSIL">
-      <HintPath>..\packages\System.Reflection.Metadata.1.6.0\lib\netstandard2.0\System.Reflection.Metadata.dll</HintPath>
-    </Reference>
     <Reference Include="System.Runtime" />
     <Reference Include="System.Runtime.Serialization" />
     <Reference Include="System.Xml.Linq" />
     <Reference Include="System.Data.DataSetExtensions" />
     <Reference Include="Microsoft.CSharp" />
     <Reference Include="System.Data" />
-    <Reference Include="System.Net.Http" />
     <Reference Include="System.Xml" />
   </ItemGroup>
   <ItemGroup>
@@ -68,10 +48,15 @@
     <Compile Include="TestExecutor.cs" />
   </ItemGroup>
   <ItemGroup>
-    <None Include="packages.config" />
+    <None Include="ReadMe.md" />
   </ItemGroup>
   <ItemGroup>
-    <None Include="ReadMe.md" />
+    <PackageReference Include="Microsoft.TestPlatform">
+      <Version>17.5.0</Version>
+    </PackageReference>
+    <PackageReference Include="Microsoft.TestPlatform.ObjectModel">
+      <Version>17.5.0</Version>
+    </PackageReference>
   </ItemGroup>
   <Import Project="$(MSBuildToolsPath)\Microsoft.CSharp.targets" />
 </Project>
\ No newline at end of file
diff --git a/Engine/Source/Programs/UnrealVS/Unreal.TestAdapter/packages.config b/Engine/Source/Programs/UnrealVS/Unreal.TestAdapter/packages.config
deleted file mode 100644
diff --git a/Engine/Source/Programs/UnrealVS/UnrealVS.2019/Properties/AssemblyInfo.cs b/Engine/Source/Programs/UnrealVS/UnrealVS.2019/Properties/AssemblyInfo.cs
--- a/Engine/Source/Programs/UnrealVS/UnrealVS.2019/Properties/AssemblyInfo.cs
+++ b/Engine/Source/Programs/UnrealVS/UnrealVS.2019/Properties/AssemblyInfo.cs
@@ -31,6 +31,6 @@ using System.Runtime.InteropServices;
 // You can specify all the values or you can default the Revision and Build Numbers 
 // by using the '*' as shown below:
 
-[assembly: AssemblyVersion("1.84.1.0")]
-[assembly: AssemblyFileVersion("1.84.1.0")]
+[assembly: AssemblyVersion("1.84.2.0")]
+[assembly: AssemblyFileVersion("1.84.2.0")]
 
diff --git a/Engine/Source/Programs/UnrealVS/UnrealVS.2022/UnrealVS.2022.csproj b/Engine/Source/Programs/UnrealVS/UnrealVS.2022/UnrealVS.2022.csproj
--- a/Engine/Source/Programs/UnrealVS/UnrealVS.2022/UnrealVS.2022.csproj
+++ b/Engine/Source/Programs/UnrealVS/UnrealVS.2022/UnrealVS.2022.csproj
@@ -95,6 +95,7 @@
     <None Include="Key.snk" />
   </ItemGroup>
   <ItemGroup>
+    <Resource Include="ReleaseNotes.txt" />
     <Resource Include="..\UnrealVS.2019\Resources\BatchBuilder.png">
       <Link>Resources\BatchBuilder.png</Link>
     </Resource>
@@ -158,9 +159,6 @@
     <PackageReference Include="Microsoft.CSharp">
       <Version>4.7.0</Version>
     </PackageReference>
-    <PackageReference Include="Microsoft.TestPlatform.ObjectModel">
-      <Version>17.3.2</Version>
-    </PackageReference>
     <PackageReference Include="Microsoft.VisualStudio.Sdk" Version="17.0.0-previews-1-31410-273">
     </PackageReference>
     <PackageReference Include="Microsoft.VisualStudio.Threading">
@@ -178,9 +176,6 @@
       <IncludeAssets>runtime; build; native; contentfiles; analyzers; buildtransitive</IncludeAssets>
       <PrivateAssets>all</PrivateAssets>
     </PackageReference>
-    <PackageReference Include="MSTest.TestAdapter">
-      <Version>2.2.10</Version>
-    </PackageReference>
     <PackageReference Include="System.Text.Json">
       <Version>6.0.0-preview.5.21301.5</Version>
     </PackageReference>
diff --git a/Engine/Source/Programs/UnrealVS/UnrealVS.Shared/P4Commands.cs b/Engine/Source/Programs/UnrealVS/UnrealVS.Shared/P4Commands.cs
--- a/Engine/Source/Programs/UnrealVS/UnrealVS.Shared/P4Commands.cs
+++ b/Engine/Source/Programs/UnrealVS/UnrealVS.Shared/P4Commands.cs
@@ -54,7 +54,7 @@ namespace UnrealVS
 		private IVsOutputWindowPane P4OutputPane;
 		private string P4WorkingDirectory;
 		private bool bPullWorkingDirectorFromP4 = true;
-		private static Mutex bPullWorkingDirectorFromP4Mutex = new Mutex();
+		private static object bPullWorkingDirectorFromP4Lock = new object();
 
 		// Exe paths
 		private string P4Exe = "C:\\Program Files\\Perforce\\p4.exe";
@@ -969,9 +969,8 @@ namespace UnrealVS
 		{
 			ThreadHelper.ThrowIfNotOnUIThread();
 
-			bPullWorkingDirectorFromP4Mutex.WaitOne();
-			try
-			{
+			lock(bPullWorkingDirectorFromP4Lock)
+			{ 
 				if (IsSolutionLoaded() && (P4WorkingDirectory == null || P4WorkingDirectory.Length < 2))
 				{
 					DTE DTE = UnrealVSPackage.Instance.DTE;
@@ -1002,10 +1001,6 @@ namespace UnrealVS
 					}
 				}
 			}
-			finally
-			{
-				bPullWorkingDirectorFromP4Mutex.ReleaseMutex();
-			}
 		}
 
 		public async Task<bool> TryP4CommandAsync(string CommandLine, bool bPullWorkingDirectoryNow = bPullWorkingDirectoryOn, bool bOutputStdOut = bOutputStdOutOn)
