UnrealVS
* Added "Generate Assembly" menu option that asks UBT to compile a single file and generate an assembly file.
* Upgraded UnrealVS.vsix to version 1.84.3

diff --git a/Engine/Extras/UnrealVS/VS2022/AtomFeed.xml b/Engine/Extras/UnrealVS/VS2022/AtomFeed.xml
--- a/Engine/Extras/UnrealVS/VS2022/AtomFeed.xml
+++ b/Engine/Extras/UnrealVS/VS2022/AtomFeed.xml
@@ -8,7 +8,7 @@
     <title type="text">UnrealVS</title>
     <summary type="text">Unreal Engine extension for Visual Studio</summary>
     <published>2022-10-21T10:01:57-07:00</published>
-    <updated>2023-02-24T13:10:00-00:00</updated>
+    <updated>2023-03-07T13:10:00-00:00</updated>
     <author>
       <name>Epic Games, Inc.</name>
     </author>
@@ -19,7 +19,7 @@
     <content type="application/octet-stream" src="UnrealVS.vsix" />
     <Vsix xmlns:xsd="http://www.w3.org/2001/XMLSchema" xmlns:xsi="http://www.w3.org/2001/XMLSchema-instance" xmlns="http://schemas.microsoft.com/developer/vsx-syndication-schema/2010">
       <Id>ddbf523f-7eb6-4887-bd51-85a714ff87eb</Id>
-      <Version>1.84.2</Version>
+      <Version>1.84.3</Version>
       <Rating xsi:nil="true" />
       <RatingCount xsi:nil="true" />
       <DownloadCount xsi:nil="true" />
diff --git a/Engine/Source/Programs/UnrealVS/ReleaseNotes.txt b/Engine/Source/Programs/UnrealVS/ReleaseNotes.txt
--- a/Engine/Source/Programs/UnrealVS/ReleaseNotes.txt
+++ b/Engine/Source/Programs/UnrealVS/ReleaseNotes.txt
@@ -1,3 +1,6 @@
+1.84.3
+* Added command to generate assembly file while compiling single file
+
 1.84.2
 * Fix to Unreal.TestAdapater to output test results to a file
 
diff --git a/Engine/Source/Programs/UnrealVS/UnrealVS.2019/Properties/AssemblyInfo.cs b/Engine/Source/Programs/UnrealVS/UnrealVS.2019/Properties/AssemblyInfo.cs
--- a/Engine/Source/Programs/UnrealVS/UnrealVS.2019/Properties/AssemblyInfo.cs
+++ b/Engine/Source/Programs/UnrealVS/UnrealVS.2019/Properties/AssemblyInfo.cs
@@ -31,6 +31,6 @@ using System.Runtime.InteropServices;
 // You can specify all the values or you can default the Revision and Build Numbers 
 // by using the '*' as shown below:
 
-[assembly: AssemblyVersion("1.84.2.0")]
-[assembly: AssemblyFileVersion("1.84.2.0")]
+[assembly: AssemblyVersion("1.84.3.0")]
+[assembly: AssemblyFileVersion("1.84.3.0")]
 
diff --git a/Engine/Source/Programs/UnrealVS/UnrealVS.2022/UnrealVS.2022.csproj b/Engine/Source/Programs/UnrealVS/UnrealVS.2022/UnrealVS.2022.csproj
--- a/Engine/Source/Programs/UnrealVS/UnrealVS.2022/UnrealVS.2022.csproj
+++ b/Engine/Source/Programs/UnrealVS/UnrealVS.2022/UnrealVS.2022.csproj
@@ -95,7 +95,6 @@
     <None Include="Key.snk" />
   </ItemGroup>
   <ItemGroup>
-    <Resource Include="ReleaseNotes.txt" />
     <Resource Include="..\UnrealVS.2019\Resources\BatchBuilder.png">
       <Link>Resources\BatchBuilder.png</Link>
     </Resource>
diff --git a/Engine/Source/Programs/UnrealVS/UnrealVS.Shared/CompileSingleFile.cs b/Engine/Source/Programs/UnrealVS/UnrealVS.Shared/CompileSingleFile.cs
--- a/Engine/Source/Programs/UnrealVS/UnrealVS.Shared/CompileSingleFile.cs
+++ b/Engine/Source/Programs/UnrealVS/UnrealVS.Shared/CompileSingleFile.cs
@@ -20,6 +20,7 @@ namespace UnrealVS
 		private const int PreprocessSingleFileButtonID = 0x1076;
 		private const int CompileSingleModuleButtonID = 0x1077;
 		private const int CompileAndProfileSingleFileButtonID = 0x1078;
+		private const int GenerateAssemblyFileButtonID = 0x1079;
 		private const int UBTSubMenuID = 0x3103;
 		private string	  FileToCompileOriginalExt = "";
 
@@ -48,6 +49,10 @@ namespace UnrealVS
 			MenuCommand CompileAndProfileSingleFileButtonCommand = new MenuCommand(new EventHandler(CompileSingleFileButtonHandler), CommandID4);
 			UnrealVSPackage.Instance.MenuCommandService.AddCommand(CompileAndProfileSingleFileButtonCommand);
 
+			CommandID CommandID5 = new CommandID(GuidList.UnrealVSCmdSet, GenerateAssemblyFileButtonID);
+			MenuCommand GenerateAssemblyFileButtonCommand = new MenuCommand(new EventHandler(CompileSingleFileButtonHandler), CommandID5);
+			UnrealVSPackage.Instance.MenuCommandService.AddCommand(GenerateAssemblyFileButtonCommand);
+
 			// add sub menu for UBT commands
 			SubMenuCommand = new OleMenuCommand(null, new CommandID(GuidList.UnrealVSCmdSet, UBTSubMenuID));
 			UnrealVSPackage.Instance.MenuCommandService.AddCommand(SubMenuCommand);
@@ -88,8 +93,9 @@ namespace UnrealVS
 			bool bIsFile = SenderSubMenuCommand.CommandID.ID != CompileSingleModuleButtonID;
 			bool bPreprocessOnly = SenderSubMenuCommand.CommandID.ID == PreprocessSingleFileButtonID;
 			bool bProfile = SenderSubMenuCommand.CommandID.ID == CompileAndProfileSingleFileButtonID;
+			bool bGenerateAssembly = SenderSubMenuCommand.CommandID.ID == GenerateAssemblyFileButtonID;
 
-			if (!TryCompileSingleFileOrModule(bIsFile, bPreprocessOnly, bProfile))
+			if (!TryCompileSingleFileOrModule(bIsFile, bPreprocessOnly, bProfile, bGenerateAssembly))
 			{
 				DTE DTE = UnrealVSPackage.Instance.DTE;
 				DTE.ExecuteCommand("Build.Compile");
@@ -126,7 +132,7 @@ namespace UnrealVS
 			return FindModuleForFile(directory, rootDirectory);
 		}
 
-		bool TryCompileSingleFileOrModule(bool bIsFile, bool bPreProcessOnly, bool bProfile)
+		bool TryCompileSingleFileOrModule(bool bIsFile, bool bPreProcessOnly, bool bProfile, bool bGenerateAssembly)
 		{
 			ThreadHelper.ThrowIfNotOnUIThread();
 			DTE DTE = UnrealVSPackage.Instance.DTE;
@@ -290,12 +296,21 @@ namespace UnrealVS
 
 			FileToCompileOriginalExt = FileToCompileExt;
 
-			string PreProcess = bPreProcessOnly ? " -NoXGE -Preprocess " : "";
+			string ExtraArgs = "";
+
+			if (bPreProcessOnly)
+			{
+				ExtraArgs = " -NoXGE -Preprocess ";
+			}
+			else if (bGenerateAssembly)
+			{
+				ExtraArgs = " -NOXGE -WithAssembly ";
+			}
 
 			// Spawn the new process
 			ChildProcess = new System.Diagnostics.Process();
 			ChildProcess.StartInfo.FileName = Path.Combine(Environment.SystemDirectory, "cmd.exe");
-			ChildProcess.StartInfo.Arguments = $"/C \"{BuildCommandLine} {PreProcess} {UBTArgument}\"\"";
+			ChildProcess.StartInfo.Arguments = $"/C \"{BuildCommandLine} {ExtraArgs} {UBTArgument}\"\"";
 			ChildProcess.StartInfo.WorkingDirectory = Path.GetDirectoryName(StartupProject.FullName);
 			ChildProcess.StartInfo.UseShellExecute = false;
 			ChildProcess.StartInfo.RedirectStandardOutput = true;
@@ -304,7 +319,7 @@ namespace UnrealVS
 			ChildProcess.OutputDataReceived += OutputHandler;
 			ChildProcess.ErrorDataReceived += OutputHandler;
 			if (bPreProcessOnly || bProfile)
-			{ 
+			{
 				// add an event handler to respond to the exit of the preprocess request
 				// and open the generated file if it exists.
 				ChildProcess.EnableRaisingEvents = true;
