UnrealVS: Allow installing for Visual Studio 2022 on ARM64.

Also upgrade all NuGet packages (required for Arm64 support).

diff --git a/Engine/Source/Programs/UnrealVS/Unreal.TestAdapter/Unreal.TestAdapter.csproj b/Engine/Source/Programs/UnrealVS/Unreal.TestAdapter/Unreal.TestAdapter.csproj
--- a/Engine/Source/Programs/UnrealVS/Unreal.TestAdapter/Unreal.TestAdapter.csproj
+++ b/Engine/Source/Programs/UnrealVS/Unreal.TestAdapter/Unreal.TestAdapter.csproj
@@ -52,10 +52,10 @@
   </ItemGroup>
   <ItemGroup>
     <PackageReference Include="Microsoft.TestPlatform">
-      <Version>17.5.0</Version>
+      <Version>17.6.3</Version>
     </PackageReference>
     <PackageReference Include="Microsoft.TestPlatform.ObjectModel">
-      <Version>17.5.0</Version>
+      <Version>17.7.0</Version>
     </PackageReference>
   </ItemGroup>
   <Import Project="$(MSBuildToolsPath)\Microsoft.CSharp.targets" />
diff --git a/Engine/Source/Programs/UnrealVS/UnrealVS.2019/UnrealVS.2019.csproj b/Engine/Source/Programs/UnrealVS/UnrealVS.2019/UnrealVS.2019.csproj
--- a/Engine/Source/Programs/UnrealVS/UnrealVS.2019/UnrealVS.2019.csproj
+++ b/Engine/Source/Programs/UnrealVS/UnrealVS.2019/UnrealVS.2019.csproj
@@ -130,10 +130,10 @@
   </ItemGroup>
   <ItemGroup>
     <PackageReference Include="Microsoft.CodeAnalysis.CSharp">
-      <Version>3.10.0</Version>
+      <Version>4.6.0</Version>
     </PackageReference>
     <PackageReference Include="Microsoft.CodeAnalysis.NetAnalyzers">
-      <Version>5.0.3</Version>
+      <Version>7.0.3</Version>
       <IncludeAssets>runtime; build; native; contentfiles; analyzers; buildtransitive</IncludeAssets>
       <PrivateAssets>all</PrivateAssets>
     </PackageReference>
@@ -141,26 +141,26 @@
       <Version>4.7.0</Version>
     </PackageReference>
     <PackageReference Include="Microsoft.VisualStudio.Sdk">
-      <Version>16.9.31025.194</Version>
+      <Version>17.6.36389</Version>
     </PackageReference>
     <PackageReference Include="Microsoft.VisualStudio.Threading">
-      <Version>16.9.54</Version>
+      <Version>17.6.40</Version>
     </PackageReference>
     <PackageReference Include="Microsoft.VisualStudio.Threading.Analyzers">
-      <Version>16.9.54</Version>
+      <Version>17.6.40</Version>
       <IncludeAssets>runtime; build; native; contentfiles; analyzers; buildtransitive</IncludeAssets>
       <PrivateAssets>all</PrivateAssets>
     </PackageReference>
     <PackageReference Include="Microsoft.VisualStudio.VCProjectEngine">
-      <Version>16.9.31025.104</Version>
+      <Version>17.6.36389</Version>
     </PackageReference>
     <PackageReference Include="Microsoft.VSSDK.BuildTools">
-      <Version>17.0.1600</Version>
+      <Version>17.6.2164</Version>
       <IncludeAssets>runtime; build; native; contentfiles; analyzers; buildtransitive</IncludeAssets>
       <PrivateAssets>all</PrivateAssets>
     </PackageReference>
     <PackageReference Include="System.Text.Json">
-      <Version>5.0.2</Version>
+      <Version>7.0.3</Version>
     </PackageReference>
   </ItemGroup>
   <Import Project="..\UnrealVS.Shared\UnrealVS.Shared.projitems" Label="Shared" />
diff --git a/Engine/Source/Programs/UnrealVS/UnrealVS.2022/UnrealVS.2022.csproj b/Engine/Source/Programs/UnrealVS/UnrealVS.2022/UnrealVS.2022.csproj
--- a/Engine/Source/Programs/UnrealVS/UnrealVS.2022/UnrealVS.2022.csproj
+++ b/Engine/Source/Programs/UnrealVS/UnrealVS.2022/UnrealVS.2022.csproj
@@ -148,35 +148,35 @@
   </ItemGroup>
   <ItemGroup>
     <PackageReference Include="Microsoft.CodeAnalysis.CSharp">
-      <Version>4.0.0-1.final</Version>
+      <Version>4.6.0</Version>
     </PackageReference>
     <PackageReference Include="Microsoft.CodeAnalysis.NetAnalyzers">
-      <Version>5.0.3</Version>
+      <Version>7.0.3</Version>
       <IncludeAssets>runtime; build; native; contentfiles; analyzers; buildtransitive</IncludeAssets>
       <PrivateAssets>all</PrivateAssets>
     </PackageReference>
     <PackageReference Include="Microsoft.CSharp">
       <Version>4.7.0</Version>
     </PackageReference>
-    <PackageReference Include="Microsoft.VisualStudio.Sdk" Version="17.0.0-previews-1-31410-273">
+    <PackageReference Include="Microsoft.VisualStudio.Sdk" Version="17.6.36389">
     </PackageReference>
     <PackageReference Include="Microsoft.VisualStudio.Threading">
-      <Version>17.0.17-alpha</Version>
+      <Version>17.6.40</Version>
     </PackageReference>
     <PackageReference Include="Microsoft.VisualStudio.Threading.Analyzers">
-      <Version>17.0.17-alpha</Version>
+      <Version>17.6.40</Version>
       <IncludeAssets>runtime; build; native; contentfiles; analyzers; buildtransitive</IncludeAssets>
       <PrivateAssets>all</PrivateAssets>
     </PackageReference>
     <PackageReference Include="Microsoft.VisualStudio.VCProjectEngine">
-      <Version>17.0.0-previews-1-31410-258</Version>
+      <Version>17.6.36389</Version>
     </PackageReference>
-    <PackageReference Include="Microsoft.VSSDK.BuildTools" Version="17.0.3154-preview3">
+    <PackageReference Include="Microsoft.VSSDK.BuildTools" Version="17.6.2164">
       <IncludeAssets>runtime; build; native; contentfiles; analyzers; buildtransitive</IncludeAssets>
       <PrivateAssets>all</PrivateAssets>
     </PackageReference>
     <PackageReference Include="System.Text.Json">
-      <Version>6.0.0-preview.5.21301.5</Version>
+      <Version>7.0.3</Version>
     </PackageReference>
   </ItemGroup>
   <ItemGroup>
