UnrealVS: Fix compile single file for external project
Show additional target.cs files

diff --git a/Engine/Extras/UnrealVS/VS2022/AtomFeed.xml b/Engine/Extras/UnrealVS/VS2022/AtomFeed.xml
--- a/Engine/Extras/UnrealVS/VS2022/AtomFeed.xml
+++ b/Engine/Extras/UnrealVS/VS2022/AtomFeed.xml
@@ -8,7 +8,7 @@
     <title type="text">UnrealVS</title>
     <summary type="text">Unreal Engine extension for Visual Studio</summary>
     <published>2022-10-21T10:01:57-07:00</published>
-    <updated>2023-07-25T15:44:40-00:00</updated>
+    <updated>2023-11-22T19:00:00-00:00</updated>
     <author>
       <name>Epic Games, Inc.</name>
     </author>
@@ -19,7 +19,7 @@
     <content type="application/octet-stream" src="UnrealVS.vsix" />
     <Vsix xmlns:xsd="http://www.w3.org/2001/XMLSchema" xmlns:xsi="http://www.w3.org/2001/XMLSchema-instance" xmlns="http://schemas.microsoft.com/developer/vsx-syndication-schema/2010">
       <Id>ddbf523f-7eb6-4887-bd51-85a714ff87eb</Id>
-      <Version>1.84.3.3</Version>
+      <Version>1.84.3.4</Version>
       <Rating xsi:nil="true" />
       <RatingCount xsi:nil="true" />
       <DownloadCount xsi:nil="true" />
diff --git a/Engine/Source/Programs/UnrealBuildTool/Modes/BuildMode.cs b/Engine/Source/Programs/UnrealBuildTool/Modes/BuildMode.cs
--- a/Engine/Source/Programs/UnrealBuildTool/Modes/BuildMode.cs
+++ b/Engine/Source/Programs/UnrealBuildTool/Modes/BuildMode.cs
@@ -804,6 +804,24 @@ namespace UnrealBuildTool
 								await actionArtifactCache.FlushChangesAsync(default); //ETSTODO
 							}
 						}
+
+						// these are parsed by external tools wishing to open this file directly
+						foreach (LinkedAction BuildAction in MergedActionsToExecute)
+						{
+							if (BuildAction.ActionType == ActionType.Compile)
+							{
+								FileItem? PreprocessedFile = BuildAction.ProducedItems.FirstOrDefault(x => x.HasExtension(".i"));
+								if (PreprocessedFile != null)
+								{
+									Logger.LogInformation("PreProcessPath: {File}", PreprocessedFile);
+								}
+								FileItem? AssemblyPath = BuildAction.ProducedItems.FirstOrDefault(x => x.HasExtension(".asm"));
+								if (AssemblyPath != null)
+								{
+									Logger.LogInformation("AssemblyPath: {File}", AssemblyPath);
+								}
+							}
+						}
 					}
 
 					// Run the deployment steps
diff --git a/Engine/Source/Programs/UnrealBuildTool/Platform/Windows/VCCompileAction.cs b/Engine/Source/Programs/UnrealBuildTool/Platform/Windows/VCCompileAction.cs
--- a/Engine/Source/Programs/UnrealBuildTool/Platform/Windows/VCCompileAction.cs
+++ b/Engine/Source/Programs/UnrealBuildTool/Platform/Windows/VCCompileAction.cs
@@ -524,12 +524,6 @@ namespace UnrealBuildTool
 			if (PreprocessedFile != null)
 			{
 				VCToolChain.AddPreprocessedFile(Arguments, PreprocessedFile, Logger);
-
-				// this is parsed by external tools wishing to open this file directly.
-				if (!bIsAnalyzing)
-				{
-					Logger.LogInformation("PreProcessPath: {File}", PreprocessedFile);
-				}
 			}
 
 			if (ObjectFile != null)
diff --git a/Engine/Source/Programs/UnrealBuildTool/ToolChain/ClangToolChain.cs b/Engine/Source/Programs/UnrealBuildTool/ToolChain/ClangToolChain.cs
--- a/Engine/Source/Programs/UnrealBuildTool/ToolChain/ClangToolChain.cs
+++ b/Engine/Source/Programs/UnrealBuildTool/ToolChain/ClangToolChain.cs
@@ -934,9 +934,6 @@ namespace UnrealBuildTool
 				// Clang does EITHER pre-process or object file.
 				Arguments.Add("-E"); // Only run the preprocessor
 				Arguments.Add("-fuse-line-directives"); // Use #line in preprocessed output
-
-				// this is parsed by external tools wishing to open this file directly.
-				Logger.LogInformation("PreProcessPath: {File}", OutputFile.AbsolutePath);
 			}
 			else if (IsAnalyzing(CompileEnvironment))
 			{
diff --git a/Engine/Source/Programs/UnrealVS/ReleaseNotes.txt b/Engine/Source/Programs/UnrealVS/ReleaseNotes.txt
--- a/Engine/Source/Programs/UnrealVS/ReleaseNotes.txt
+++ b/Engine/Source/Programs/UnrealVS/ReleaseNotes.txt
@@ -1,3 +1,9 @@
+1.84.3.4
+* Bug fix for incorrect paths when compiling a single file
+* Open .asm after generating assembly
+* Disable remaining compile single buttons if a valid file isn't selected
+* Bug fix to show additional .target.cs files
+
 1.84.3.3
 * Bug fix for low level test projects - do not auto prefix project path in command line
 
diff --git a/Engine/Source/Programs/UnrealVS/UnrealVS.2019/Properties/AssemblyInfo.cs b/Engine/Source/Programs/UnrealVS/UnrealVS.2019/Properties/AssemblyInfo.cs
--- a/Engine/Source/Programs/UnrealVS/UnrealVS.2019/Properties/AssemblyInfo.cs
+++ b/Engine/Source/Programs/UnrealVS/UnrealVS.2019/Properties/AssemblyInfo.cs
@@ -31,6 +31,6 @@ using System.Runtime.InteropServices;
 // You can specify all the values or you can default the Revision and Build Numbers 
 // by using the '*' as shown below:
 
-[assembly: AssemblyVersion("1.84.3.3")]
-[assembly: AssemblyFileVersion("1.84.3.3")]
+[assembly: AssemblyVersion("1.84.3.4")]
+[assembly: AssemblyFileVersion("1.84.3.4")]
 
diff --git a/Engine/Source/Programs/UnrealVS/UnrealVS.Shared/CompileSingleFile.cs b/Engine/Source/Programs/UnrealVS/UnrealVS.Shared/CompileSingleFile.cs
--- a/Engine/Source/Programs/UnrealVS/UnrealVS.Shared/CompileSingleFile.cs
+++ b/Engine/Source/Programs/UnrealVS/UnrealVS.Shared/CompileSingleFile.cs
@@ -46,11 +46,13 @@ namespace UnrealVS
 			UnrealVSPackage.Instance.MenuCommandService.AddCommand(CompileSingleModuleButtonCommand);
 
 			CommandID CommandID4 = new CommandID(GuidList.UnrealVSCmdSet, CompileAndProfileSingleFileButtonID);
-			MenuCommand CompileAndProfileSingleFileButtonCommand = new MenuCommand(new EventHandler(CompileSingleFileButtonHandler), CommandID4);
+			var CompileAndProfileSingleFileButtonCommand = new OleMenuCommand(new EventHandler(CompileSingleFileButtonHandler), CommandID4);
+			CompileAndProfileSingleFileButtonCommand.BeforeQueryStatus += CompileSingleFileButtonCommand_BeforeQueryStatus;
 			UnrealVSPackage.Instance.MenuCommandService.AddCommand(CompileAndProfileSingleFileButtonCommand);
 
 			CommandID CommandID5 = new CommandID(GuidList.UnrealVSCmdSet, GenerateAssemblyFileButtonID);
-			MenuCommand GenerateAssemblyFileButtonCommand = new MenuCommand(new EventHandler(CompileSingleFileButtonHandler), CommandID5);
+			var GenerateAssemblyFileButtonCommand = new OleMenuCommand(new EventHandler(CompileSingleFileButtonHandler), CommandID5);
+			GenerateAssemblyFileButtonCommand.BeforeQueryStatus += CompileSingleFileButtonCommand_BeforeQueryStatus;
 			UnrealVSPackage.Instance.MenuCommandService.AddCommand(GenerateAssemblyFileButtonCommand);
 
 			// add sub menu for UBT commands
@@ -212,7 +214,21 @@ namespace UnrealVS
 			string FileToCompileExt = Path.GetExtension(FileToCompile);
 
 			string CompilingText;
-			string UBTArgument;
+			List<string> UBTArguments = new List<string>
+			{
+				"-WorkingDir=\"$(MSBuildProjectDirectory)\"",
+			};
+			if (bPreProcessOnly)
+			{
+				UBTArguments.Add("-NoXGE -NoSNDBS -NoFASTBuild");
+				UBTArguments.Add("-Preprocess");
+			}
+			else if (bGenerateAssembly)
+			{
+				UBTArguments.Add("-NoXGE -NoSNDBS -NoFASTBuild");
+				UBTArguments.Add("-WithAssembly");
+			}
+
 			if (bIsFile)
 			{
 				if (!ValidExtensions.Contains(FileToCompileExt.ToLowerInvariant()))
@@ -222,7 +238,7 @@ namespace UnrealVS
 				}
 
 				CompilingText = FileToCompile;
-				UBTArgument = $"-singlefile=\"{FileToCompile}";
+				UBTArguments.Add($"-SingleFile=\"{FileToCompile}\"");
 			}
 			else
 			{
@@ -230,9 +246,10 @@ namespace UnrealVS
 				if (ModuleName == null)
 				{
 					MessageBox.Show($"Can't find module for for {FileToCompile} to compile.", "Invalid Module", MessageBoxButtons.OK);
+					return true;
 				}
 				CompilingText = ModuleName;
-				UBTArgument = $"-Module=\"{ModuleName}";
+				UBTArguments.Add($"-Module=\"{ModuleName}\"");
 			}
 
 			// If there's already a build in progress, don't let another one start
@@ -258,15 +275,23 @@ namespace UnrealVS
 			DTE.Events.BuildEvents.OnBuildBegin += BuildEvents_OnBuildBegin;
 
 			// Create a delegate for handling output messages
-			string PPLine = String.Empty;
+			List<string> PreprocessedFiles = new List<string>();
+			List<string> AssemblyFiles = new List<string>();
 			void OutputHandler(object Sender, DataReceivedEventArgs Args) 
 			{ 
 				if (Args.Data != null) 
 				{ 
-					BuildOutputPane.OutputStringThreadSafe($"1>  {Args.Data}{Environment.NewLine}"); 
 					if (Args.Data.Contains("PreProcessPath:"))
 					{
-						PPLine = Args.Data;
+						PreprocessedFiles.Add(Args.Data.Replace("PreProcessPath:", "").Trim());
+					}
+					else if (Args.Data.Contains("AssemblyPath:"))
+					{
+						AssemblyFiles.Add(Args.Data.Replace("AssemblyPath:", "").Trim());
+					}
+					else
+					{
+						BuildOutputPane.OutputStringThreadSafe($"1>  {Args.Data}{Environment.NewLine}");
 					}
 				} 
 			}
@@ -290,43 +315,31 @@ namespace UnrealVS
 
 			string SolutionDir = Path.GetDirectoryName(UnrealVSPackage.Instance.SolutionFilepath).TrimEnd(Path.DirectorySeparatorChar) + Path.DirectorySeparatorChar;
 			// Get the build command line and escape any environment variables that we use
-			string BuildCommandLine = ActiveNMakeTool.BuildCommandLine;
-			BuildCommandLine = BuildCommandLine.Replace("$(SolutionDir)", SolutionDir);
-			BuildCommandLine = BuildCommandLine.Replace("$(ProjectName)", VCStartupProject.Name);
-			BuildCommandLine = BuildCommandLine.Replace("$(BuildBatchScript)", "..\\..\\Build\\BatchFiles\\Build.bat");
+			string BuildCommandLine = ActiveVCConfiguration.Evaluate(ActiveNMakeTool.BuildCommandLine);
+			string UBTArgument = ActiveVCConfiguration.Evaluate(string.Join(" ", UBTArguments));
+			string WorkingDirectory = ActiveVCConfiguration.Evaluate("$(MSBuildProjectDirectory)");
 
 			FileToCompileOriginalExt = FileToCompileExt;
 
-			string ExtraArgs = "";
-
-			if (bPreProcessOnly)
-			{
-				ExtraArgs = " -NoXGE -Preprocess ";
-			}
-			else if (bGenerateAssembly)
-			{
-				ExtraArgs = " -NOXGE -WithAssembly ";
-			}
-
 			// Spawn the new process
 			ChildProcess = new System.Diagnostics.Process();
 			ChildProcess.StartInfo.FileName = Path.Combine(Environment.SystemDirectory, "cmd.exe");
-			ChildProcess.StartInfo.Arguments = $"/C \"{BuildCommandLine} {ExtraArgs} {UBTArgument}\"\"";
-			ChildProcess.StartInfo.WorkingDirectory = Path.GetDirectoryName(StartupProject.FullName);
+			ChildProcess.StartInfo.Arguments = $"/C \"{BuildCommandLine} {UBTArgument}\"";
+			ChildProcess.StartInfo.WorkingDirectory = WorkingDirectory;
 			ChildProcess.StartInfo.UseShellExecute = false;
 			ChildProcess.StartInfo.RedirectStandardOutput = true;
 			ChildProcess.StartInfo.RedirectStandardError = true;
 			ChildProcess.StartInfo.CreateNoWindow = true;
 			ChildProcess.OutputDataReceived += OutputHandler;
 			ChildProcess.ErrorDataReceived += OutputHandler;
-			if (bPreProcessOnly || bProfile)
+			if (bPreProcessOnly || bGenerateAssembly  || bProfile)
 			{
 				// add an event handler to respond to the exit of the preprocess request
 				// and open the generated file if it exists.
 				ChildProcess.EnableRaisingEvents = true;
-				ChildProcess.Exited += new EventHandler((s, e) => PreprocessExitHandler(PPLine, bProfile));
+				ChildProcess.Exited += new EventHandler((s, e) => PreprocessExitHandler(PreprocessedFiles, AssemblyFiles, bProfile));
 			}
-			
+
 			ChildProcess.Start();
 			ChildProcess.BeginOutputReadLine();
 			ChildProcess.BeginErrorReadLine();
@@ -334,7 +347,7 @@ namespace UnrealVS
 			return true;
 		}
 
-		private void PreprocessExitHandler(string PPLine, bool bIsProfiling)
+		private void PreprocessExitHandler(IEnumerable<string> PreprocessedFiles, IEnumerable<string> AssemblyFiles, bool bIsProfiling)
 		{
 			if (bIsProfiling)
 			{
@@ -346,14 +359,20 @@ namespace UnrealVS
 			}
 
 			// not all compile actions support pre-process - check it exists
-			if (PPLine.Contains("PreProcessPath:"))
+			foreach (string PreprocessedFile in PreprocessedFiles)
+			{
+				ThreadHelper.JoinableTaskFactory.Run(async () =>
+				{
+					await ThreadHelper.JoinableTaskFactory.SwitchToMainThreadAsync();
+					OpenPreprocessedFile(PreprocessedFile);
+				});
+			}
+			foreach (string AssemblyFile in AssemblyFiles)
 			{
-				string PPFullPath = PPLine.Replace("PreProcessPath:", "").Trim();
-
 				ThreadHelper.JoinableTaskFactory.Run(async () =>
 				{
 					await ThreadHelper.JoinableTaskFactory.SwitchToMainThreadAsync();
-					OpenPreprocessedFile(PPFullPath);
+					OpenAssemblyFile(AssemblyFile);
 				});
 			}
 		}
@@ -362,11 +381,6 @@ namespace UnrealVS
 		{
 			ThreadHelper.ThrowIfNotOnUIThread();
 
-			IVsOutputWindowPane BuildOutputPane = UnrealVSPackage.Instance.GetOutputPane();
-			DTE DTE = UnrealVSPackage.Instance.DTE;
-
-			BuildOutputPane.OutputStringThreadSafe($"1>  PPFullPath: {PPFullPath}{Environment.NewLine}");
-
 			if (File.Exists(PPFullPath))
 			{
 				// if the file exists, rename it to isolate the file and have its extension be the original to maintain syntax highlighting
@@ -377,7 +391,17 @@ namespace UnrealVS
 
 				File.Copy(PPFullPath, RenamedFile, true /*overwrite*/);
 
-				DTE.ExecuteCommand("File.OpenFile", $"\"{RenamedFile}\"");
+				UnrealVSPackage.Instance.DTE.ExecuteCommand("File.OpenFile", $"\"{RenamedFile}\"");
+			}
+		}
+
+		private void OpenAssemblyFile(string AsmFullPath)
+		{
+			ThreadHelper.ThrowIfNotOnUIThread();
+
+			if (File.Exists(AsmFullPath))
+			{
+				UnrealVSPackage.Instance.DTE.ExecuteCommand("File.OpenFile", $"\"{AsmFullPath}\"");
 			}
 		}
 
