UnrealVS: Fix an exception when editing command line arguments for C# projects without a Properties folder.

diff --git a/Engine/Source/Programs/UnrealVS/ReleaseNotes.txt b/Engine/Source/Programs/UnrealVS/ReleaseNotes.txt
--- a/Engine/Source/Programs/UnrealVS/ReleaseNotes.txt
+++ b/Engine/Source/Programs/UnrealVS/ReleaseNotes.txt
@@ -1,3 +1,6 @@
+1.84.3.5
+* Fix null reference exception when setting command line for C# project without a Properties directory
+
 1.84.3.4
 * Bug fix for incorrect paths when compiling a single file
 * Open .asm after generating assembly
diff --git a/Engine/Source/Programs/UnrealVS/UnrealVS.2019/Properties/AssemblyInfo.cs b/Engine/Source/Programs/UnrealVS/UnrealVS.2019/Properties/AssemblyInfo.cs
--- a/Engine/Source/Programs/UnrealVS/UnrealVS.2019/Properties/AssemblyInfo.cs
+++ b/Engine/Source/Programs/UnrealVS/UnrealVS.2019/Properties/AssemblyInfo.cs
@@ -31,6 +31,6 @@ using System.Runtime.InteropServices;
 // You can specify all the values or you can default the Revision and Build Numbers 
 // by using the '*' as shown below:
 
-[assembly: AssemblyVersion("1.84.3.4")]
-[assembly: AssemblyFileVersion("1.84.3.4")]
+[assembly: AssemblyVersion("1.84.3.5")]
+[assembly: AssemblyFileVersion("1.84.3.5")]
 
diff --git a/Engine/Source/Programs/UnrealVS/UnrealVS.2022/UnrealVS.2022.csproj b/Engine/Source/Programs/UnrealVS/UnrealVS.2022/UnrealVS.2022.csproj
--- a/Engine/Source/Programs/UnrealVS/UnrealVS.2022/UnrealVS.2022.csproj
+++ b/Engine/Source/Programs/UnrealVS/UnrealVS.2022/UnrealVS.2022.csproj
@@ -203,6 +203,7 @@
     <ErrorReport>prompt</ErrorReport>
     <CodeAnalysisRuleSet>MinimumRecommendedRules.ruleset</CodeAnalysisRuleSet>
     <NoWarn>MSB3277</NoWarn>
+    <CopyVsixExtensionFiles>False</CopyVsixExtensionFiles>
   </PropertyGroup>
   <PropertyGroup Condition=" '$(Configuration)|$(Platform)' == 'Release|AnyCPU' ">
     <OutputPath>bin\Release\</OutputPath>
diff --git a/Engine/Source/Programs/UnrealVS/UnrealVS.Shared/CommandLineEditor.cs b/Engine/Source/Programs/UnrealVS/UnrealVS.Shared/CommandLineEditor.cs
--- a/Engine/Source/Programs/UnrealVS/UnrealVS.Shared/CommandLineEditor.cs
+++ b/Engine/Source/Programs/UnrealVS/UnrealVS.Shared/CommandLineEditor.cs
@@ -434,7 +434,8 @@ namespace UnrealVS
 										activeDebugProfile = null;
 									}
 
-									string launchSettingsPath = Path.Combine(Path.GetDirectoryName(SelectedStartupProject.FileName), "Properties", "launchSettings.json");
+									string launchSettingsDir = Path.Combine(Path.GetDirectoryName(SelectedStartupProject.FileName), "Properties");
+									string launchSettingsPath = Path.Combine(launchSettingsDir, "launchSettings.json");
 									LaunchSettingsJson settings = null;
 									if (File.Exists(launchSettingsPath))
 									{
@@ -481,6 +482,8 @@ namespace UnrealVS
 									{
 										WriteIndented = true,
 									});
+
+									Directory.CreateDirectory(launchSettingsDir);
 									File.WriteAllText(launchSettingsPath, json);
 								}
 							}
