UnrealVS: Fix hang when p4.exe isn't installed.

#rnx

diff --git a/Engine/Extras/UnrealVS/VS2022/AtomFeed.xml b/Engine/Extras/UnrealVS/VS2022/AtomFeed.xml
--- a/Engine/Extras/UnrealVS/VS2022/AtomFeed.xml
+++ b/Engine/Extras/UnrealVS/VS2022/AtomFeed.xml
@@ -2,7 +2,7 @@
 <feed xmlns="http://www.w3.org/2005/Atom">
   <title type="text">UnrealVS Gallery</title>
   <id>uuid:45b4571a-2981-49f1-9f70-04cfa96db5c8;id=1</id>
-  <updated>2022-10-21T18:08:01Z</updated>
+  <updated>2024-02-28T17:50:26Z</updated>
   <entry>
     <id>ddbf523f-7eb6-4887-bd51-85a714ff87eb</id>
     <title type="text">UnrealVS</title>
@@ -19,7 +19,7 @@
     <content type="application/octet-stream" src="UnrealVS.vsix" />
     <Vsix xmlns:xsd="http://www.w3.org/2001/XMLSchema" xmlns:xsi="http://www.w3.org/2001/XMLSchema-instance" xmlns="http://schemas.microsoft.com/developer/vsx-syndication-schema/2010">
       <Id>ddbf523f-7eb6-4887-bd51-85a714ff87eb</Id>
-      <Version>1.84.3.4</Version>
+      <Version>1.84.3.6</Version>
       <Rating xsi:nil="true" />
       <RatingCount xsi:nil="true" />
       <DownloadCount xsi:nil="true" />
diff --git a/Engine/Source/Programs/UnrealVS/ReleaseNotes.txt b/Engine/Source/Programs/UnrealVS/ReleaseNotes.txt
--- a/Engine/Source/Programs/UnrealVS/ReleaseNotes.txt
+++ b/Engine/Source/Programs/UnrealVS/ReleaseNotes.txt
@@ -1,3 +1,6 @@
+1.84.3.6
+* Fix P4 checkout queue causing hang if p4.exe is not found in path.
+
 1.84.3.5
 * Fix null reference exception when setting command line for C# project without a Properties directory
 
diff --git a/Engine/Source/Programs/UnrealVS/UnrealVS.2019/Properties/AssemblyInfo.cs b/Engine/Source/Programs/UnrealVS/UnrealVS.2019/Properties/AssemblyInfo.cs
--- a/Engine/Source/Programs/UnrealVS/UnrealVS.2019/Properties/AssemblyInfo.cs
+++ b/Engine/Source/Programs/UnrealVS/UnrealVS.2019/Properties/AssemblyInfo.cs
@@ -31,6 +31,6 @@ using System.Runtime.InteropServices;
 // You can specify all the values or you can default the Revision and Build Numbers 
 // by using the '*' as shown below:
 
-[assembly: AssemblyVersion("1.84.3.5")]
-[assembly: AssemblyFileVersion("1.84.3.5")]
+[assembly: AssemblyVersion("1.84.3.6")]
+[assembly: AssemblyFileVersion("1.84.3.6")]
 
diff --git a/Engine/Source/Programs/UnrealVS/UnrealVS.Shared/P4Commands.cs b/Engine/Source/Programs/UnrealVS/UnrealVS.Shared/P4Commands.cs
--- a/Engine/Source/Programs/UnrealVS/UnrealVS.Shared/P4Commands.cs
+++ b/Engine/Source/Programs/UnrealVS/UnrealVS.Shared/P4Commands.cs
@@ -203,8 +203,16 @@ namespace UnrealVS
 			string CheckoutQueueFile = GetCheckoutQueueFileName();
 			if (CheckoutQueueFile != null)
 			{
-				P4OutputPane.OutputStringThreadSafe($"UnrealVS started. Using checkout queue: {CheckoutQueueFile}" + Environment.NewLine);
-				PulseCheckoutQueue(CheckoutQueueFile);
+				P4OutputPane.OutputStringThreadSafe("UnrealVS started" + Environment.NewLine);
+				if (File.Exists(P4Exe))
+				{
+					P4OutputPane.OutputStringThreadSafe($"Using checkout queue: {CheckoutQueueFile}" + Environment.NewLine);
+					PulseCheckoutQueue(CheckoutQueueFile);
+				}
+				else
+				{
+					P4OutputPane.OutputStringThreadSafe($"Unable to find Perforce executable ({P4Exe}). Perforce checkout queue will be disabled." + Environment.NewLine);
+				}
 			}
 		}
 
@@ -890,6 +898,8 @@ namespace UnrealVS
 							CheckoutQueueFiles.Add(QueueFile);
 						}
 					}
+
+					await Task.Delay(TimeSpan.FromSeconds(10.0));
 				}
 			}
 		}
@@ -1076,7 +1086,7 @@ namespace UnrealVS
 			// Create a delegate for handling output messages
 			ChildProcess.OutputDataReceived += (s, a) => { if (a.Data != null) StdOutSB.AppendLine(a.Data); };
 			ChildProcess.ErrorDataReceived += (s, a) => { if (a.Data != null) StdErrSB.AppendLine(a.Data); };
-				ChildProcess.EnableRaisingEvents = true;
+			ChildProcess.EnableRaisingEvents = true;
 
 			TaskCompletionSource<bool> ProcessExitTaskSource = new TaskCompletionSource<bool>();
 			ChildProcess.Exited += (s, a) =>
@@ -1089,7 +1099,15 @@ namespace UnrealVS
 				ChildProcessList.Add(ChildProcess);
 			}
 
-			ChildProcess.Start();
+			try
+			{
+				ChildProcess.Start();
+			}
+			catch (Exception ex)
+			{
+				return (false, "", $"Unable to launch {CmdPath} ({ex.Message})");
+			}
+
 			ChildProcess.BeginOutputReadLine();
 			ChildProcess.BeginErrorReadLine();
 
