PR #11563: UnrealVS: Prefix encrypted Perforce server port address with "ssl:"

diff --git a/Engine/Extras/UnrealVS/VS2022/AtomFeed.xml b/Engine/Extras/UnrealVS/VS2022/AtomFeed.xml
--- a/Engine/Extras/UnrealVS/VS2022/AtomFeed.xml
+++ b/Engine/Extras/UnrealVS/VS2022/AtomFeed.xml
@@ -8,7 +8,7 @@
     <title type="text">UnrealVS</title>
     <summary type="text">Unreal Engine extension for Visual Studio</summary>
     <published>2022-10-21T10:01:57-07:00</published>
-    <updated>2023-11-22T19:00:00-00:00</updated>
+    <updated>2024-02-28T18:30:00-00:00</updated>
     <author>
       <name>Epic Games, Inc.</name>
     </author>
diff --git a/Engine/Source/Programs/UnrealVS/ReleaseNotes.txt b/Engine/Source/Programs/UnrealVS/ReleaseNotes.txt
--- a/Engine/Source/Programs/UnrealVS/ReleaseNotes.txt
+++ b/Engine/Source/Programs/UnrealVS/ReleaseNotes.txt
@@ -1,3 +1,6 @@
+1.84.3.7
+* Prefix encrypted Perforce server port address with "ssl:"
+
 1.84.3.6
 * Fix P4 checkout queue causing hang if p4.exe is not found in path.
 
diff --git a/Engine/Source/Programs/UnrealVS/UnrealVS.Shared/P4Commands.cs b/Engine/Source/Programs/UnrealVS/UnrealVS.Shared/P4Commands.cs
--- a/Engine/Source/Programs/UnrealVS/UnrealVS.Shared/P4Commands.cs
+++ b/Engine/Source/Programs/UnrealVS/UnrealVS.Shared/P4Commands.cs
@@ -959,17 +959,24 @@ namespace UnrealVS
 			Regex PortPattern = new Regex(@"Server address: (?<port>.*)$", RegexOptions.Compiled | RegexOptions.Multiline);
 			Regex ClientPattern = new Regex(@"Client name: (?<client>.*)$", RegexOptions.Compiled | RegexOptions.Multiline);
 			Regex ClientStream = new Regex(@"Client stream: (?<stream>.*)$", RegexOptions.Compiled | RegexOptions.Multiline);
+			Regex ServerEncryption = new Regex(@"Server encryption: (?<encryption>.*)$", RegexOptions.Compiled | RegexOptions.Multiline);
 
 			System.Text.RegularExpressions.Match UserMatch = UserPattern.Match(CaptureP4Output);
 			System.Text.RegularExpressions.Match PortMatch = PortPattern.Match(CaptureP4Output);
 			System.Text.RegularExpressions.Match ClientMatch = ClientPattern.Match(CaptureP4Output);
 			System.Text.RegularExpressions.Match StreamMatch = ClientStream.Match(CaptureP4Output);
+			System.Text.RegularExpressions.Match EncryptionMatch = ServerEncryption.Match(CaptureP4Output);
 
 			Port = PortMatch.Groups["port"].Value.Trim();
 			Username = UserMatch.Groups["user"].Value.Trim();
 			Client = ClientMatch.Groups["client"].Value.Trim();
 			Stream = StreamMatch.Groups["stream"].Value.Trim();
 
+			if (EncryptionMatch.Groups["encryption"].Value.Trim().Equals("encrypted", StringComparison.OrdinalIgnoreCase) && !Port.StartsWith("ssl:", StringComparison.OrdinalIgnoreCase))
+			{
+				Port = "ssl:" + Port;
+			}
+
 			UserInfoComplete = string.Format(" -p {0} -u {1} -c {2} ", Port, Username, Client);
 
 			P4OutputPane.OutputStringThreadSafe("GetUserInfoStringFull : " + UserInfoComplete + Environment.NewLine);
