[UnrealVS]
* Fixed so CompileScore works for non-ubt paths as well

diff --git a/Engine/Extras/UnrealVS/VS2022/AtomFeed.xml b/Engine/Extras/UnrealVS/VS2022/AtomFeed.xml
--- a/Engine/Extras/UnrealVS/VS2022/AtomFeed.xml
+++ b/Engine/Extras/UnrealVS/VS2022/AtomFeed.xml
@@ -8,7 +8,7 @@
     <title type="text">UnrealVS</title>
     <summary type="text">Unreal Engine extension for Visual Studio</summary>
     <published>2022-10-21T10:01:57-07:00</published>
-    <updated>2024-03-07T18:30:00-00:00</updated>
+    <updated>2024-03-15T18:30:00-00:00</updated>
     <author>
       <name>Epic Games, Inc.</name>
     </author>
@@ -19,7 +19,7 @@
     <content type="application/octet-stream" src="UnrealVS.vsix" />
     <Vsix xmlns:xsd="http://www.w3.org/2001/XMLSchema" xmlns:xsi="http://www.w3.org/2001/XMLSchema-instance" xmlns="http://schemas.microsoft.com/developer/vsx-syndication-schema/2010">
       <Id>ddbf523f-7eb6-4887-bd51-85a714ff87eb</Id>
-      <Version>1.84.3.8</Version>
+      <Version>1.84.3.9</Version>
       <Rating xsi:nil="true" />
       <RatingCount xsi:nil="true" />
       <DownloadCount xsi:nil="true" />
diff --git a/Engine/Source/Programs/UnrealVS/ReleaseNotes.txt b/Engine/Source/Programs/UnrealVS/ReleaseNotes.txt
--- a/Engine/Source/Programs/UnrealVS/ReleaseNotes.txt
+++ b/Engine/Source/Programs/UnrealVS/ReleaseNotes.txt
@@ -1,3 +1,6 @@
+1.84.3.9
+* Fixed issues with compile single file with profile. Did not work if selected project was a c# project. Should now also work with non-ue projects
+
 1.84.3.8
 * Fixed issues with context menu items.
 
diff --git a/Engine/Source/Programs/UnrealVS/UnrealVS.2019/Properties/AssemblyInfo.cs b/Engine/Source/Programs/UnrealVS/UnrealVS.2019/Properties/AssemblyInfo.cs
--- a/Engine/Source/Programs/UnrealVS/UnrealVS.2019/Properties/AssemblyInfo.cs
+++ b/Engine/Source/Programs/UnrealVS/UnrealVS.2019/Properties/AssemblyInfo.cs
@@ -31,6 +31,6 @@ using System.Runtime.InteropServices;
 // You can specify all the values or you can default the Revision and Build Numbers 
 // by using the '*' as shown below:
 
-[assembly: AssemblyVersion("1.84.3.8")]
-[assembly: AssemblyFileVersion("1.84.3.8")]
+[assembly: AssemblyVersion("1.84.3.9")]
+[assembly: AssemblyFileVersion("1.84.3.9")]
 
diff --git a/Engine/Source/Programs/UnrealVS/UnrealVS.Shared/CompileSingleFile.cs b/Engine/Source/Programs/UnrealVS/UnrealVS.Shared/CompileSingleFile.cs
--- a/Engine/Source/Programs/UnrealVS/UnrealVS.Shared/CompileSingleFile.cs
+++ b/Engine/Source/Programs/UnrealVS/UnrealVS.Shared/CompileSingleFile.cs
@@ -11,6 +11,7 @@ using System.IO;
 using System.Windows.Forms;
 using System.Text;
 using System.Linq;
+using System.Threading.Tasks;
 
 namespace UnrealVS
 {
@@ -100,7 +101,34 @@ namespace UnrealVS
 			if (!TryCompileSingleFileOrModule(bIsFile, bPreprocessOnly, bProfile, bGenerateAssembly))
 			{
 				DTE DTE = UnrealVSPackage.Instance.DTE;
+
+				if (bProfile)
+					if (!StartTrace())
+						return;
+
 				DTE.ExecuteCommand("Build.Compile");
+
+				if (bProfile)
+				{
+					var WaitThread = new System.Threading.Thread(() =>
+					{
+						bool Done = false;
+						while (!Done)
+						{
+							ThreadHelper.JoinableTaskFactory.Run(async () =>
+							{
+								await ThreadHelper.JoinableTaskFactory.SwitchToMainThreadAsync();
+								if (DTE.Solution.SolutionBuild.BuildState != vsBuildState.vsBuildStateDone)
+									return;
+								StopTrace();
+								Done = true;
+
+							});
+						}
+					})
+					{ Priority = System.Threading.ThreadPriority.Lowest };
+					WaitThread.Start();
+				}
 			}
 		}
 
@@ -134,6 +162,43 @@ namespace UnrealVS
 			return FindModuleForFile(directory, rootDirectory);
 		}
 
+		bool StartTrace()
+		{
+			ThreadHelper.ThrowIfNotOnUIThread();
+			DTE DTE = UnrealVSPackage.Instance.DTE;
+
+			var Commands = DTE.Commands.Cast<Command>();
+#pragma warning disable VSTHRD010
+			Command Command = Commands.FirstOrDefault((C) => C.Name == "CompileScore.StartTrace");
+#pragma warning restore VSTHRD010
+			if (Command == null || !Command.IsAvailable)
+			{
+				MessageBox.Show($"CompileSingleFileAndProfile requires a CompileScore visual studio extension that has CompileScore.StartTrace/CompileScore.StopTrace installed", "UnrealVS - Missing CompileScore extension ", MessageBoxButtons.OK);
+				return false;
+			}
+			Object CustomIn = null;
+			Object CustomOut = null;
+			DTE.Commands.Raise(Command.Guid, Command.ID, ref CustomIn, ref CustomOut);
+			//DTE.ExecuteCommand("CompileScore.StartTrace");
+			IVsOutputWindowPane BuildOutputPane = UnrealVSPackage.Instance.GetOutputPane();
+			if (BuildOutputPane == null)
+			{
+				BuildOutputPane.OutputStringThreadSafe($"1>------ Started trace for CompileScore ------{Environment.NewLine}");
+			}
+			return true;
+		}
+
+		void StopTrace()
+		{
+			ThreadHelper.ThrowIfNotOnUIThread();
+			UnrealVSPackage.Instance.DTE.ExecuteCommand("CompileScore.StopTrace");
+			IVsOutputWindowPane BuildOutputPane = UnrealVSPackage.Instance.GetOutputPane();
+			if (BuildOutputPane != null)
+			{
+				BuildOutputPane.OutputStringThreadSafe($"1>------ Stopped trace for CompileScore ------{Environment.NewLine}");
+			}
+		}
+
 		bool TryCompileSingleFileOrModule(bool bIsFile, bool bPreProcessOnly, bool bProfile, bool bGenerateAssembly)
 		{
 			ThreadHelper.ThrowIfNotOnUIThread();
@@ -268,8 +333,6 @@ namespace UnrealVS
 			// Set up the output pane
 			BuildOutputPane.Activate();
 			BuildOutputPane.Clear();
-			BuildOutputPane.OutputStringThreadSafe($"1>------ Build started: Project: {StartupProject.Name}, Configuration: {ActiveConfiguration.ConfigurationName} {ActiveConfiguration.PlatformName} ------{Environment.NewLine}");
-			BuildOutputPane.OutputStringThreadSafe($"1>  Compiling {CompilingText}{Environment.NewLine}");
 
 			// Set up event handlers 
 			DTE.Events.BuildEvents.OnBuildBegin += BuildEvents_OnBuildBegin;
@@ -297,21 +360,11 @@ namespace UnrealVS
 			}
 
 			if (bProfile)
-			{
-				var Commands = DTE.Commands.Cast<Command>();
-				#pragma warning disable VSTHRD010
-				Command Command = Commands.FirstOrDefault((C) => C.Name == "CompileScore.StartTrace");
-				#pragma warning restore VSTHRD010
-				if (Command == null || !Command.IsAvailable)
-				{
-					MessageBox.Show($"CompileSingleFileAndProfile requires a CompileScore visual studio extension that has CompileScore.StartTrace/CompileScore.StopTrace installed", "UnrealVS - Missing CompileScore extension ", MessageBoxButtons.OK);
+				if (!StartTrace())
 					return true;
-				}
-				Object CustomIn = null;
-				Object CustomOut = null;
-				DTE.Commands.Raise(Command.Guid, Command.ID, ref CustomIn, ref CustomOut);
-				//DTE.ExecuteCommand("CompileScore.StartTrace");
-			}
+
+			BuildOutputPane.OutputStringThreadSafe($"1>------ Build started: Project: {StartupProject.Name}, Configuration: {ActiveConfiguration.ConfigurationName} {ActiveConfiguration.PlatformName} ------{Environment.NewLine}");
+			BuildOutputPane.OutputStringThreadSafe($"1>  Compiling {CompilingText}{Environment.NewLine}");
 
 			string SolutionDir = Path.GetDirectoryName(UnrealVSPackage.Instance.SolutionFilepath).TrimEnd(Path.DirectorySeparatorChar) + Path.DirectorySeparatorChar;
 			// Get the build command line and escape any environment variables that we use
@@ -337,7 +390,7 @@ namespace UnrealVS
 				// add an event handler to respond to the exit of the preprocess request
 				// and open the generated file if it exists.
 				ChildProcess.EnableRaisingEvents = true;
-				ChildProcess.Exited += new EventHandler((s, e) => PreprocessExitHandler(PreprocessedFiles, AssemblyFiles, bProfile));
+				ChildProcess.Exited += new EventHandler((s, e) => UbtProcessExitHandler(PreprocessedFiles, AssemblyFiles, bProfile));
 			}
 
 			ChildProcess.Start();
@@ -347,14 +400,14 @@ namespace UnrealVS
 			return true;
 		}
 
-		private void PreprocessExitHandler(IEnumerable<string> PreprocessedFiles, IEnumerable<string> AssemblyFiles, bool bIsProfiling)
+		private void UbtProcessExitHandler(IEnumerable<string> PreprocessedFiles, IEnumerable<string> AssemblyFiles, bool bIsProfiling)
 		{
 			if (bIsProfiling)
 			{
 				ThreadHelper.JoinableTaskFactory.Run(async () =>
 				{
 					await ThreadHelper.JoinableTaskFactory.SwitchToMainThreadAsync();
-					UnrealVSPackage.Instance.DTE.ExecuteCommand("CompileScore.StopTrace");
+					StopTrace();
 				});
 			}
 
