[UnrealVs]
* Added UbaVisualizer panel to UnrealVs

diff --git a/Engine/Extras/UnrealVS/VS2022/AtomFeed.xml b/Engine/Extras/UnrealVS/VS2022/AtomFeed.xml
--- a/Engine/Extras/UnrealVS/VS2022/AtomFeed.xml
+++ b/Engine/Extras/UnrealVS/VS2022/AtomFeed.xml
@@ -8,7 +8,7 @@
     <title type="text">UnrealVS</title>
     <summary type="text">Unreal Engine extension for Visual Studio</summary>
     <published>2022-10-21T10:01:57-07:00</published>
-    <updated>2024-04-05T18:30:00-00:00</updated>
+    <updated>2024-07-20T09:37:00-00:00</updated>
     <author>
       <name>Epic Games, Inc.</name>
     </author>
@@ -19,7 +19,7 @@
     <content type="application/octet-stream" src="UnrealVS.vsix" />
     <Vsix xmlns:xsd="http://www.w3.org/2001/XMLSchema" xmlns:xsi="http://www.w3.org/2001/XMLSchema-instance" xmlns="http://schemas.microsoft.com/developer/vsx-syndication-schema/2010">
       <Id>ddbf523f-7eb6-4887-bd51-85a714ff87eb</Id>
-      <Version>1.84.4.0</Version>
+      <Version>1.84.5.0</Version>
       <Rating xsi:nil="true" />
       <RatingCount xsi:nil="true" />
       <DownloadCount xsi:nil="true" />
diff --git a/Engine/Source/Programs/UnrealVS/ReleaseNotes.txt b/Engine/Source/Programs/UnrealVS/ReleaseNotes.txt
--- a/Engine/Source/Programs/UnrealVS/ReleaseNotes.txt
+++ b/Engine/Source/Programs/UnrealVS/ReleaseNotes.txt
@@ -1,3 +1,6 @@
+1.84.5.0
+* Added UbaVisualizer panel. If open solution is under UE engine path this panel will launch UbaVisualizer.exe embedded and listen to uba traces spawned by vs instance
+
 1.84.4.0
 * Fixed a race condition at initialization time that would make Visual Studio unresponsive.
 
diff --git a/Engine/Source/Programs/UnrealVS/UnrealVS.2019/Properties/AssemblyInfo.cs b/Engine/Source/Programs/UnrealVS/UnrealVS.2019/Properties/AssemblyInfo.cs
--- a/Engine/Source/Programs/UnrealVS/UnrealVS.2019/Properties/AssemblyInfo.cs
+++ b/Engine/Source/Programs/UnrealVS/UnrealVS.2019/Properties/AssemblyInfo.cs
@@ -31,6 +31,6 @@ using System.Runtime.InteropServices;
 // You can specify all the values or you can default the Revision and Build Numbers 
 // by using the '*' as shown below:
 
-[assembly: AssemblyVersion("1.84.4.0")]
-[assembly: AssemblyFileVersion("1.84.4.0")]
+[assembly: AssemblyVersion("1.84.5.0")]
+[assembly: AssemblyFileVersion("1.84.5.0")]
 
diff --git a/Engine/Source/Programs/UnrealVS/UnrealVS.2022/UnrealVS.2022.csproj b/Engine/Source/Programs/UnrealVS/UnrealVS.2022/UnrealVS.2022.csproj
--- a/Engine/Source/Programs/UnrealVS/UnrealVS.2022/UnrealVS.2022.csproj
+++ b/Engine/Source/Programs/UnrealVS/UnrealVS.2022/UnrealVS.2022.csproj
@@ -176,7 +176,7 @@
       <PrivateAssets>all</PrivateAssets>
     </PackageReference>
     <PackageReference Include="System.Text.Json">
-      <Version>7.0.3</Version>
+      <Version>9.0.0-preview.6.24327.7</Version>
     </PackageReference>
   </ItemGroup>
   <ItemGroup>
diff --git a/Engine/Source/Programs/UnrealVS/UnrealVS.Shared/UbaVisualizer.cs b/Engine/Source/Programs/UnrealVS/UnrealVS.Shared/UbaVisualizer.cs
new file mode 100644
--- /dev/null
+++ b/Engine/Source/Programs/UnrealVS/UnrealVS.Shared/UbaVisualizer.cs
@@ -0,0 +1,193 @@
+﻿// Copyright Epic Games, Inc. All Rights Reserved.
+
+using EnvDTE80;
+using Microsoft.VisualStudio.Shell;
+using Microsoft.VisualStudio.Shell.Interop;
+using System;
+using System.ComponentModel.Design;
+using System.Diagnostics;
+using System.Diagnostics.CodeAnalysis;
+using System.Runtime.InteropServices;
+using System.Security;
+using System.Windows.Forms;
+using System.Windows.Interop;
+
+namespace UnrealVS
+{
+	delegate IntPtr WndProc(IntPtr hWnd, uint msg, IntPtr wParam, IntPtr lParam);
+	class UbaVisualizer
+	{
+		const int UbaVisualizerButtonID = 0x1338;
+
+		public UbaVisualizer()
+		{
+			// UbaVisualizerButton
+			var CommandID = new CommandID(GuidList.UnrealVSCmdSet, UbaVisualizerButtonID);
+			var UbaVisualizerButtonCommand = new MenuCommand(new EventHandler(UbaVisualizerButtonHandler), CommandID);
+			UnrealVSPackage.Instance.MenuCommandService.AddCommand(UbaVisualizerButtonCommand);
+
+			UnrealVSPackage.Instance.OnSolutionOpened += () => VisualizerControl?.HandleSolutionChanged();
+			UnrealVSPackage.Instance.OnSolutionClosed += () => VisualizerControl?.HandleSolutionChanged();
+			UnrealVSPackage.Instance.OnProjectLoaded += (p) => VisualizerControl?.HandleSolutionChanged();
+			UnrealVSPackage.Instance.OnProjectUnloading += (p) => VisualizerControl?.HandleSolutionChanged();
+			//UnrealVSPackage.Instance.OnDocumentActivated += (d) => VisualizerControl?.HandleDocumentActivated(d);
+		}
+
+		/// Called when 'UbaVisualizer' button is clicked
+		void UbaVisualizerButtonHandler(object Sender, EventArgs Args)
+		{
+			ThreadHelper.ThrowIfNotOnUIThread();
+
+			ToolWindowPane ToolWindow = UnrealVSPackage.Instance.FindToolWindow(typeof(UbaVisualizerWindow), 0, true);
+			if ((null == ToolWindow) || (null == ToolWindow.Frame))
+			{
+				throw new NotSupportedException(Resources.ToolWindowCreateError);
+			}
+			IVsWindowFrame ToolWindowFrame = (IVsWindowFrame)ToolWindow.Frame;
+			Microsoft.VisualStudio.ErrorHandler.ThrowOnFailure(ToolWindowFrame.Show());
+		}
+
+		UbaVisualizerWindowControl VisualizerControl
+		{
+			get
+			{
+				ToolWindowPane ToolWindow = UnrealVSPackage.Instance.FindToolWindow(typeof(UbaVisualizerWindow), 0, true);
+				if ((null == ToolWindow) || (null == ToolWindow.Frame))
+					return null;
+				var UbaVisualizer = (UbaVisualizerWindow)ToolWindow;
+				return (UbaVisualizerWindowControl)UbaVisualizer.Content;
+			}
+		}
+	}
+
+	[Guid("b205a03c-566f-4dc3-9e2f-99e2df765f1e")]
+	public class UbaVisualizerWindow : ToolWindowPane
+	{
+		public UbaVisualizerWindow() : base(null)
+		{
+			this.Caption = "UbaVisualizer";
+
+			this.Content = new UbaVisualizerWindowControl(this);
+		}
+
+		[global::System.Diagnostics.CodeAnalysis.SuppressMessageAttribute("", "VSTHRD010")]
+		protected override bool PreProcessMessage(ref Message m)
+        {
+			if (m.Msg == 0x0100)//WM_KEYDOWN)
+			{
+				if (m.WParam == new IntPtr(0x1B)) // VK_ESCAPE
+				{
+					//((FileBrowserWindowControl)this.Content).HandleEscape();
+					return true;
+				}
+				else if (m.WParam == new IntPtr(0x70)) // VK_F1
+				{
+					//((FileBrowserWindowControl)this.Content).HandleF1();
+					return true;
+				}
+				else if (m.WParam == new IntPtr(0x74)) // VK_F5
+                {
+					//((FileBrowserWindowControl)this.Content).HandleF5();
+					return true;
+                }
+			}
+			return base.PreProcessMessage(ref m);
+        }
+    }
+
+	[SuppressMessage("Microsoft.Usage", "CA2216:DisposableTypesShouldDeclareFinalizer")]
+	internal class UbaVisualizerHost : Hwnd​Host
+	{
+		string VisualizerExe;
+		IntPtr HwndVisualizer;
+		WndProc HwndProc;
+		IntPtr HwndStatic;
+
+		public UbaVisualizerHost(string visualizerExe)
+		{
+			VisualizerExe = visualizerExe;
+		}
+
+		protected override HandleRef BuildWindowCore(HandleRef hwndParent)
+		{
+			HwndProc = new WndProc(CustomWndProc);
+
+			NativeMethods.WNDCLASSEX wc = new NativeMethods.WNDCLASSEX();
+			wc.cbSize = Marshal.SizeOf(typeof(NativeMethods.WNDCLASSEX));
+			wc.style = 0x0001 | 0x0002 | 0x0008; // CS_HREDRAW | CS_VREDRAW | CS_DBLCLKS;
+			wc.hbrBackground = IntPtr.Zero;// (IntPtr)2; //Black background
+			wc.cbClsExtra = 0;
+			wc.cbWndExtra = 0;
+			wc.hInstance = Marshal.GetHINSTANCE(this.GetType().Module);
+			wc.hIcon = IntPtr.Zero;
+			wc.hCursor = NativeMethods.LoadCursorW(IntPtr.Zero, (IntPtr)32512); // IDC_ARROW
+			wc.lpszMenuName = null;
+			wc.lpszClassName = "UbaVisualizerWindow";
+			wc.lpfnWndProc = Marshal.GetFunctionPointerForDelegate(HwndProc);
+			wc.hIconSm = IntPtr.Zero;
+			NativeMethods.RegisterClassExW(ref wc);
+
+			int style = 0x50000000; // WS_CHILD | WS_VISIBLE
+
+			IntPtr hwndHost = NativeMethods.CreateWindowEx(0, "UbaVisualizerWindow", "", style, 0, 0, 1, 1, hwndParent.Handle, IntPtr.Zero, IntPtr.Zero, 0);
+			string hwndStr = hwndHost.ToString("X").ToLower();
+
+			ProcessStartInfo startInfo = new ProcessStartInfo(VisualizerExe);
+			startInfo.UseShellExecute = false;
+			startInfo.Arguments = $"-parent={hwndStr}";
+			Process.Start(startInfo);
+
+			return new HandleRef(this, hwndHost);
+		}
+
+		protected override void DestroyWindowCore(HandleRef hwnd)
+		{
+			NativeMethods.DestroyWindow(hwnd.Handle);
+		}
+
+		protected IntPtr CustomWndProc(IntPtr hwnd, uint msg, IntPtr wParam, IntPtr lParam)
+		{
+			bool hasVisualizer = false;
+			bool shouldResize = msg == 0x0005;
+
+			if (HwndVisualizer != IntPtr.Zero && NativeMethods.IsWindow(HwndVisualizer))
+			{
+				hasVisualizer = true;
+			}
+			else
+			{
+				if (msg == 0x0444) // Special UbaVisualizer message
+				{
+					NativeMethods.DestroyWindow(HwndStatic);
+					HwndStatic = IntPtr.Zero;
+					HwndVisualizer = lParam;
+					hasVisualizer = true;
+					shouldResize = true;
+				}
+				else if (HwndVisualizer != IntPtr.Zero)
+				{
+					HwndVisualizer = IntPtr.Zero;
+					int style = 0x50000000; // WS_CHILD | WS_VISIBLE
+					string text = $"Waiting for UbaVisualizer process... (parent hwnd: {hwnd.ToString("X").ToLower()})";
+					HwndStatic = NativeMethods.CreateWindowEx(0, "static", text, style, 0, 0, (int)ActualWidth, (int)ActualHeight, hwnd, IntPtr.Zero, IntPtr.Zero, 0);
+				}
+			}
+
+			IntPtr res = NativeMethods.DefWindowProc(hwnd, msg, wParam, lParam);
+
+			if (shouldResize)
+			{
+				if (hasVisualizer)
+				{
+					uint flags = 0x0040 | 0x4000; // SWP_SHOWWINDOW | SWP_ASYNCWINDOWPOS
+					NativeMethods.SetWindowPos(HwndVisualizer, IntPtr.Zero, 0, 0, (int)ActualWidth, (int)ActualHeight, flags);
+				}
+				else if (HwndStatic != IntPtr.Zero)
+				{
+					NativeMethods.SetWindowPos(HwndStatic, IntPtr.Zero, 0, 0, (int)ActualWidth, (int)ActualHeight, 0x0040);
+				}
+			}
+			return res;
+		}
+	}
+}
diff --git a/Engine/Source/Programs/UnrealVS/UnrealVS.Shared/UbaVisualizerWindowControl.xaml b/Engine/Source/Programs/UnrealVS/UnrealVS.Shared/UbaVisualizerWindowControl.xaml
new file mode 100644
--- /dev/null
+++ b/Engine/Source/Programs/UnrealVS/UnrealVS.Shared/UbaVisualizerWindowControl.xaml
@@ -0,0 +1,13 @@
+<UserControl x:Class="UnrealVS.UbaVisualizerWindowControl"
+             xmlns="http://schemas.microsoft.com/winfx/2006/xaml/presentation"
+             xmlns:x="http://schemas.microsoft.com/winfx/2006/xaml"
+             xmlns:mc="http://schemas.openxmlformats.org/markup-compatibility/2006"
+             xmlns:d="http://schemas.microsoft.com/expression/blend/2008"
+             xmlns:vsshell="clr-namespace:Microsoft.VisualStudio.Shell;assembly=Microsoft.VisualStudio.Shell.15.0"
+             Background="{DynamicResource {x:Static vsshell:VsBrushes.ToolWindowBackgroundKey}}"
+             Foreground="{DynamicResource {x:Static vsshell:VsBrushes.WindowTextKey}}"
+             mc:Ignorable="d"
+             d:DesignHeight="300" d:DesignWidth="500"
+             Name="FileBrowser">
+	<Border Name="ControlHostElement" />
+</UserControl>
diff --git a/Engine/Source/Programs/UnrealVS/UnrealVS.Shared/UbaVisualizerWindowControl.xaml.cs b/Engine/Source/Programs/UnrealVS/UnrealVS.Shared/UbaVisualizerWindowControl.xaml.cs
new file mode 100644
--- /dev/null
+++ b/Engine/Source/Programs/UnrealVS/UnrealVS.Shared/UbaVisualizerWindowControl.xaml.cs
@@ -0,0 +1,106 @@
+using System;
+using System.Diagnostics;
+using System.IO;
+using System.Windows.Controls;
+
+namespace UnrealVS
+{
+	public partial class UbaVisualizerWindowControl : UserControl
+	{
+		string LastSolutionPath = "Uninitialized";
+		string DefaultText = "Starts when opening solution that contains Engine path";
+
+		public UbaVisualizerWindowControl(UbaVisualizerWindow window)
+		{
+			InitializeComponent();
+
+
+			Environment.SetEnvironmentVariable("UBA_OWNER_ID", "vs");
+			Environment.SetEnvironmentVariable("UBA_OWNER_PID", Process.GetCurrentProcess().Id.ToString());
+
+			HandleSolutionChanged();
+		}
+
+		void SetChild(string visualizerPath, string text)
+		{
+			var old = ControlHostElement.Child as UbaVisualizerHost;
+
+			if (visualizerPath != null)
+			{
+				ControlHostElement.Child = new UbaVisualizerHost(visualizerPath);
+			}
+			else if (ControlHostElement.Child is TextBlock textBlock)
+			{
+				textBlock.Text = text;
+				return;
+			}
+			else
+			{
+				textBlock = new TextBlock();
+				textBlock.HorizontalAlignment = System.Windows.HorizontalAlignment.Center;
+				textBlock.VerticalAlignment = System.Windows.VerticalAlignment.Center;
+				textBlock.Text = text;
+				ControlHostElement.Child = textBlock;
+			}
+
+			if (old != null)
+			{
+				old.Dispose();
+			}
+		}
+
+		public void HandleSolutionChanged()
+		{
+			string path = UnrealVSPackage.Instance.SolutionFilepath;
+			if (path == null)
+			{
+				path = string.Empty;
+			}
+
+			if (LastSolutionPath == path)
+			{
+				return;
+			}
+			LastSolutionPath = path;
+
+			if (path == string.Empty)
+			{
+				SetChild(null, DefaultText);
+			}
+
+			string engineStr = @"\Engine\";
+			int engineIndex = path.IndexOf(engineStr);
+			if (engineIndex == -1)
+			{
+				path = Path.GetDirectoryName(path);
+				engineIndex = path.Length;
+				path += engineStr;
+			}
+
+			string fullPath = path.Substring(0, engineIndex + engineStr.Length) + @"Binaries\Win64\UnrealBuildAccelerator\x64\UbaVisualizer.exe";
+			FileInfo fi = new FileInfo(fullPath);
+			if (!fi.Exists)
+			{
+				SetChild(null, DefaultText);
+				return;
+			}
+
+			var versionInfo = FileVersionInfo.GetVersionInfo(fullPath);
+			var pv = versionInfo.ProductVersion;
+			if (string.IsNullOrEmpty(pv))
+			{
+				pv = "Uba_v0.x.x";
+			}
+
+			if (pv.Length <= 5 || !pv.StartsWith("Uba_v") || pv[5] == '0')
+			{
+				SetChild(null, $"UbaVisualizer.exe found is too old to be embedded in visual studio ({pv})");
+				return;
+			}
+
+
+			SetChild(fullPath, "");
+		}
+	}
+
+}
\ No newline at end of file
diff --git a/Engine/Source/Programs/UnrealVS/UnrealVS.Shared/UnrealVSPackage.cs b/Engine/Source/Programs/UnrealVS/UnrealVS.Shared/UnrealVSPackage.cs
--- a/Engine/Source/Programs/UnrealVS/UnrealVS.Shared/UnrealVSPackage.cs
+++ b/Engine/Source/Programs/UnrealVS/UnrealVS.Shared/UnrealVSPackage.cs
@@ -55,10 +55,17 @@ namespace UnrealVS
 	// This attribute registers a tool window exposed by this package.
 	[ProvideToolWindow(typeof(BatchBuilderToolWindow))]
 
+	// File browser
+	[ProvideToolWindow(typeof(FileBrowserWindow))]
+
+	// UbaVisualizer
+	[ProvideToolWindow(typeof(UbaVisualizerWindow))]
+	[ProvideToolWindowVisibility(typeof(UbaVisualizerWindow), VSConstants.UICONTEXT.SolutionExists_string)]
+
 	// This attribute registers an options page for the package.
 	[ProvideOptionPage(typeof(UnrealVsOptions), ExtensionName, "General", 101, 102, true)]
 	[ProvideSolutionProperties(GuidList.UnrealVSPackageString)]
-	[ProvideToolWindow(typeof(FileBrowserWindow))]
+
 	/// <summary>
 	/// UnrealVSPackage implements Package abstract class.  This is the main class that is registered
 	/// with Visual Studio shell and serves as the entry point into our extension
@@ -283,6 +290,9 @@ namespace UnrealVS
 			// Create 'FileBrowser' instance
 			FileBrowser = new FileBrowser();
 
+			// Create 'UbaVisualizer' instance
+			UbaVisualizer = new UbaVisualizer();
+
 			// Call parent implementation
 			base.Initialize();
 
@@ -1093,6 +1103,9 @@ namespace UnrealVS
 		/// FileBrowser feature
 		private FileBrowser FileBrowser;
 
+		/// UbaVisualizer feature
+		private UbaVisualizer UbaVisualizer;
+
 		/// CompileSingleFile feature
 		private CompileSingleFile CompileSingleFile;
 
@@ -1178,5 +1191,61 @@ namespace UnrealVS
 
 		[DllImport("user32.dll")]
 		public static extern uint MapVirtualKey(uint uCode, MapType uMapType);
+
+		[DllImport("user32.dll", EntryPoint = "CreateWindowEx", CharSet = CharSet.Unicode)]
+		internal static extern IntPtr CreateWindowEx(int dwExStyle,
+													  string lpszClassName,
+													  string lpszWindowName,
+													  int style,
+													  int x, int y,
+													  int width, int height,
+													  IntPtr hwndParent,
+													  IntPtr hMenu,
+													  IntPtr hInst,
+													  [MarshalAs(UnmanagedType.AsAny)] object pvParam);
+
+		[DllImport("user32.dll", EntryPoint = "DestroyWindow", CharSet = CharSet.Unicode)]
+		internal static extern bool DestroyWindow(IntPtr hwnd);
+
+		[DllImport("user32.dll", EntryPoint = "SetWindowTextW", CharSet = CharSet.Unicode)]
+		internal static extern bool SetWindowTextW(IntPtr hWnd, string lpString);
+
+		[DllImport("user32.dll", EntryPoint = "SetWindowPos", CharSet = CharSet.Unicode)]
+		internal static extern bool SetWindowPos(IntPtr hWnd, IntPtr hWndInsertAfter, int X, int Y, int cx, int cy, uint uFlags);
+
+		[DllImport("user32.dll", EntryPoint = "IsWindow", CharSet = CharSet.Unicode)]
+		internal static extern bool IsWindow(IntPtr hWnd);
+
+		[DllImport("user32.dll", EntryPoint = "PostMessageW", CharSet = CharSet.Unicode)]
+		internal static extern bool PostMessageW(IntPtr hWnd, int Msg, IntPtr wParam, IntPtr lParam);
+
+		[DllImport("user32.dll")]
+		internal static extern IntPtr LoadCursorW(IntPtr hInstance, IntPtr lpCursorName);
+		
+
+				[DllImport("user32.dll", EntryPoint = "RegisterClassExW", SetLastError = true)]
+		internal static extern UInt16 RegisterClassExW(ref WNDCLASSEX lpWndClass);
+
+		[DllImport("user32.dll")]
+		internal static extern IntPtr DefWindowProc(IntPtr hWnd, uint uMsg, IntPtr wParam, IntPtr lParam);
+
+		[StructLayout(LayoutKind.Sequential, CharSet = CharSet.Unicode)]
+		internal struct WNDCLASSEX
+		{
+			[MarshalAs(UnmanagedType.U4)]
+			public int cbSize;
+			[MarshalAs(UnmanagedType.U4)]
+			public int style;
+			public IntPtr lpfnWndProc; // not WndProc
+			public int cbClsExtra;
+			public int cbWndExtra;
+			public IntPtr hInstance;
+			public IntPtr hIcon;
+			public IntPtr hCursor;
+			public IntPtr hbrBackground;
+			public string lpszMenuName;
+			public string lpszClassName;
+			public IntPtr hIconSm;
+		}
 	}
 }
