[UnrealVs]
* Added missing copyright comment

diff --git a/Engine/Source/Programs/UnrealVS/UnrealVS.Shared/UbaVisualizerWindowControl.xaml.cs b/Engine/Source/Programs/UnrealVS/UnrealVS.Shared/UbaVisualizerWindowControl.xaml.cs
--- a/Engine/Source/Programs/UnrealVS/UnrealVS.Shared/UbaVisualizerWindowControl.xaml.cs
+++ b/Engine/Source/Programs/UnrealVS/UnrealVS.Shared/UbaVisualizerWindowControl.xaml.cs
@@ -1,3 +1,5 @@
+// Copyright Epic Games, Inc. All Rights Reserved.
+
 using System;
 using System.Diagnostics;
 using System.IO;
