UnrealVS: Update NuGet packages and remove VS2019 support

#rnx

diff --git a/Engine/Extras/UnrealVS/VS2022/AtomFeed.xml b/Engine/Extras/UnrealVS/VS2022/AtomFeed.xml
--- a/Engine/Extras/UnrealVS/VS2022/AtomFeed.xml
+++ b/Engine/Extras/UnrealVS/VS2022/AtomFeed.xml
@@ -8,7 +8,7 @@
     <title type="text">UnrealVS</title>
     <summary type="text">Unreal Engine extension for Visual Studio</summary>
     <published>2022-10-21T10:01:57-07:00</published>
-    <updated>2024-07-20T09:37:00-00:00</updated>
+    <updated>2024-07-23T09:37:00-00:00</updated>
     <author>
       <name>Epic Games, Inc.</name>
     </author>
@@ -19,7 +19,7 @@
     <content type="application/octet-stream" src="UnrealVS.vsix" />
     <Vsix xmlns:xsd="http://www.w3.org/2001/XMLSchema" xmlns:xsi="http://www.w3.org/2001/XMLSchema-instance" xmlns="http://schemas.microsoft.com/developer/vsx-syndication-schema/2010">
       <Id>ddbf523f-7eb6-4887-bd51-85a714ff87eb</Id>
-      <Version>1.84.5.0</Version>
+      <Version>1.84.6.0</Version>
       <Rating xsi:nil="true" />
       <RatingCount xsi:nil="true" />
       <DownloadCount xsi:nil="true" />
diff --git a/Engine/Source/Programs/UnrealVS/ReleaseNotes.txt b/Engine/Source/Programs/UnrealVS/ReleaseNotes.txt
--- a/Engine/Source/Programs/UnrealVS/ReleaseNotes.txt
+++ b/Engine/Source/Programs/UnrealVS/ReleaseNotes.txt
@@ -1,3 +1,7 @@
+1.84.6.0
+* Update NuGet packages
+* Remove VS2019 support
+
 1.84.5.0
 * Added UbaVisualizer panel. If open solution is under UE engine path this panel will launch UbaVisualizer.exe embedded and listen to uba traces spawned by vs instance
 
diff --git a/Engine/Source/Programs/UnrealVS/UnrealVS.2019/Properties/AssemblyInfo.cs b/Engine/Source/Programs/UnrealVS/UnrealVS.2019/Properties/AssemblyInfo.cs
deleted file mode 100644
diff --git a/Engine/Source/Programs/UnrealVS/UnrealVS.2022/Properties/AssemblyInfo.cs b/Engine/Source/Programs/UnrealVS/UnrealVS.2022/Properties/AssemblyInfo.cs
new file mode 100644
--- /dev/null
+++ b/Engine/Source/Programs/UnrealVS/UnrealVS.2022/Properties/AssemblyInfo.cs
@@ -0,0 +1,36 @@
+// Copyright Epic Games, Inc. All Rights Reserved.
+
+using System;
+using System.Reflection;
+using System.Resources;
+using System.Runtime.InteropServices;
+
+// General Information about an assembly is controlled through the following 
+// set of attributes. Change these attribute values to modify the information
+// associated with an assembly.
+[assembly: AssemblyTitle("UnrealVS")]
+[assembly: AssemblyDescription("")]
+[assembly: AssemblyConfiguration("")]
+
+[assembly: AssemblyCompany("Epic Games, Inc.")]
+[assembly: AssemblyProduct("UnrealVS")]
+[assembly: AssemblyCopyright("Copyright Epic Games, Inc.")]
+[assembly: AssemblyTrademark("")]
+[assembly: AssemblyCulture("")]
+[assembly: ComVisible(false)]
+[assembly: CLSCompliant(false)]
+[assembly: NeutralResourcesLanguage("en-US")]
+
+// Version information for an assembly consists of the following four values:
+//
+//      Major Version
+//      Minor Version 
+//      Build Number
+//      Revision
+//
+// You can specify all the values or you can default the Revision and Build Numbers 
+// by using the '*' as shown below:
+
+[assembly: AssemblyVersion("1.84.6.0")]
+[assembly: AssemblyFileVersion("1.84.6.0")]
+
diff --git a/Engine/Source/Programs/UnrealVS/UnrealVS.2022/UnrealVS.2022.csproj b/Engine/Source/Programs/UnrealVS/UnrealVS.2022/UnrealVS.2022.csproj
--- a/Engine/Source/Programs/UnrealVS/UnrealVS.2022/UnrealVS.2022.csproj
+++ b/Engine/Source/Programs/UnrealVS/UnrealVS.2022/UnrealVS.2022.csproj
@@ -46,10 +46,9 @@
     <BootstrapperEnabled>true</BootstrapperEnabled>
   </PropertyGroup>
   <ItemGroup>
-    <VSCTCompile Include="..\UnrealVS.2019\UnrealVS.vsct">
+    <VSCTCompile Include="UnrealVS.vsct">
       <ResourceName>Menus.ctmenu</ResourceName>
       <SubType>Designer</SubType>
-      <Link>UnrealVS.vsct</Link>
     </VSCTCompile>
   </ItemGroup>
   <ItemGroup>
@@ -67,8 +66,7 @@
     <Reference Include="WindowsBase" />
   </ItemGroup>
   <ItemGroup>
-    <Compile Include="..\UnrealVS.2019\Properties\AssemblyInfo.cs">
-      <Link>Properties\AssemblyInfo.cs</Link>
+    <Compile Include="Properties\AssemblyInfo.cs">
     </Compile>
     <Compile Include="Resources.Designer.cs">
       <AutoGen>True</AutoGen>
@@ -95,27 +93,24 @@
     <None Include="Key.snk" />
   </ItemGroup>
   <ItemGroup>
-    <Resource Include="..\UnrealVS.2019\Resources\BatchBuilder.png">
-      <Link>Resources\BatchBuilder.png</Link>
+    <Resource Include="Resources\BatchBuilder.png">
     </Resource>
-    <Resource Include="..\UnrealVS.2019\Resources\BatchBuilderBackground.png">
-      <Link>Resources\BatchBuilderBackground.png</Link>
+    <Resource Include="Resources\BatchBuilderBackground.png">
     </Resource>
-    <Resource Include="..\UnrealVS.2019\Resources\Icon.ico">
-      <Link>Resources\Icon.ico</Link>
+    <Resource Include="Resources\Icon.ico">
     </Resource>
-    <Content Include="..\UnrealVS.2019\Resources\Icon.png">
-      <Link>Resources\Icon.png</Link>
+    <Resource Include="..\..\..\..\Extras\UnrealVS\VS2022\AtomFeed.xml">
+      <Link>AtomFeed.xml</Link>
+    </Resource>
+    <Content Include="Resources\Icon.png">
       <CopyToOutputDirectory>Always</CopyToOutputDirectory>
       <IncludeInVSIX>true</IncludeInVSIX>
     </Content>
-    <Resource Include="..\UnrealVS.2019\Resources\Preview.png">
-      <Link>Resources\Preview.png</Link>
+    <Resource Include="Resources\Preview.png">
       <IncludeInVSIX>true</IncludeInVSIX>
       <CopyToOutputDirectory>Always</CopyToOutputDirectory>
     </Resource>
-    <Resource Include="..\UnrealVS.2019\Resources\RefreshProjects.png">
-      <Link>Resources\RefreshProjects.png</Link>
+    <Resource Include="Resources\RefreshProjects.png">
     </Resource>
     <None Include="source.extension.vsixmanifest" />
     <Content Include="..\ReleaseNotes.txt">
@@ -148,35 +143,35 @@
   </ItemGroup>
   <ItemGroup>
     <PackageReference Include="Microsoft.CodeAnalysis.CSharp">
-      <Version>4.6.0</Version>
+      <Version>4.10.0</Version>
     </PackageReference>
     <PackageReference Include="Microsoft.CodeAnalysis.NetAnalyzers">
-      <Version>7.0.3</Version>
+      <Version>8.0.0</Version>
       <IncludeAssets>runtime; build; native; contentfiles; analyzers; buildtransitive</IncludeAssets>
       <PrivateAssets>all</PrivateAssets>
     </PackageReference>
     <PackageReference Include="Microsoft.CSharp">
       <Version>4.7.0</Version>
     </PackageReference>
-    <PackageReference Include="Microsoft.VisualStudio.Sdk" Version="17.6.36389">
+    <PackageReference Include="Microsoft.VisualStudio.Sdk" Version="17.10.40171">
     </PackageReference>
     <PackageReference Include="Microsoft.VisualStudio.Threading">
-      <Version>17.6.40</Version>
+      <Version>17.10.48</Version>
     </PackageReference>
     <PackageReference Include="Microsoft.VisualStudio.Threading.Analyzers">
-      <Version>17.6.40</Version>
+      <Version>17.10.48</Version>
       <IncludeAssets>runtime; build; native; contentfiles; analyzers; buildtransitive</IncludeAssets>
       <PrivateAssets>all</PrivateAssets>
     </PackageReference>
     <PackageReference Include="Microsoft.VisualStudio.VCProjectEngine">
-      <Version>17.6.36389</Version>
+      <Version>17.10.40170</Version>
     </PackageReference>
-    <PackageReference Include="Microsoft.VSSDK.BuildTools" Version="17.6.2164">
+    <PackageReference Include="Microsoft.VSSDK.BuildTools" Version="17.11.414">
       <IncludeAssets>runtime; build; native; contentfiles; analyzers; buildtransitive</IncludeAssets>
       <PrivateAssets>all</PrivateAssets>
     </PackageReference>
     <PackageReference Include="System.Text.Json">
-      <Version>9.0.0-preview.6.24327.7</Version>
+      <Version>8.0.4</Version>
     </PackageReference>
   </ItemGroup>
   <ItemGroup>
diff --git a/Engine/Source/Programs/UnrealVS/UnrealVS.Shared/P4Commands.cs b/Engine/Source/Programs/UnrealVS/UnrealVS.Shared/P4Commands.cs
--- a/Engine/Source/Programs/UnrealVS/UnrealVS.Shared/P4Commands.cs
+++ b/Engine/Source/Programs/UnrealVS/UnrealVS.Shared/P4Commands.cs
@@ -785,6 +785,8 @@ namespace UnrealVS
 			}
 		}
 
+		[System.Diagnostics.CodeAnalysis.SuppressMessage("Usage", "VSTHRD003:Avoid awaiting foreign Tasks", Justification = "TODO")]
+		[System.Diagnostics.CodeAnalysis.SuppressMessage("Reliability", "VSSDK007:ThreadHelper.JoinableTaskFactory.RunAsync", Justification = "TODO")]
 		private void PulseCheckoutQueue(string QueueFile)
 		{
 			TaskCompletionSource<bool> StartTask = null;
diff --git a/Engine/Source/Programs/UnrealVS/UnrealVS.Shared/Utils.cs b/Engine/Source/Programs/UnrealVS/UnrealVS.Shared/Utils.cs
--- a/Engine/Source/Programs/UnrealVS/UnrealVS.Shared/Utils.cs
+++ b/Engine/Source/Programs/UnrealVS/UnrealVS.Shared/Utils.cs
@@ -315,8 +315,10 @@ namespace UnrealVS
 
 				var MatchedSolutionConfig =
 					(from SolutionConfiguration2 Sc in SolutionConfigs select Sc).FirstOrDefault(
-						Sc =>
-						String.CompareOrdinal(Sc.Name, SolutionConfig) == 0 && String.CompareOrdinal(Sc.PlatformName, SolutionPlatform) == 0);
+						Sc => {
+							ThreadHelper.ThrowIfNotOnUIThread();
+							return String.CompareOrdinal(Sc.Name, SolutionConfig) == 0 && String.CompareOrdinal(Sc.PlatformName, SolutionPlatform) == 0;
+						});
 
 				if (MatchedSolutionConfig != null)
 				{
