UnrealVS: Fix exception with older versions of VS

#rnx

diff --git a/Engine/Extras/UnrealVS/VS2022/AtomFeed.xml b/Engine/Extras/UnrealVS/VS2022/AtomFeed.xml
--- a/Engine/Extras/UnrealVS/VS2022/AtomFeed.xml
+++ b/Engine/Extras/UnrealVS/VS2022/AtomFeed.xml
@@ -19,7 +19,7 @@
     <content type="application/octet-stream" src="UnrealVS.vsix" />
     <Vsix xmlns:xsd="http://www.w3.org/2001/XMLSchema" xmlns:xsi="http://www.w3.org/2001/XMLSchema-instance" xmlns="http://schemas.microsoft.com/developer/vsx-syndication-schema/2010">
       <Id>ddbf523f-7eb6-4887-bd51-85a714ff87eb</Id>
-      <Version>1.84.6.0</Version>
+      <Version>1.84.7.0</Version>
       <Rating xsi:nil="true" />
       <RatingCount xsi:nil="true" />
       <DownloadCount xsi:nil="true" />
diff --git a/Engine/Source/Programs/UnrealVS/ReleaseNotes.txt b/Engine/Source/Programs/UnrealVS/ReleaseNotes.txt
--- a/Engine/Source/Programs/UnrealVS/ReleaseNotes.txt
+++ b/Engine/Source/Programs/UnrealVS/ReleaseNotes.txt
@@ -1,4 +1,7 @@
-1.84.6.0
+1.84.7.0
+* Fix exception for older VS versions
+
+1.84.7.0
 * Update NuGet packages
 * Remove VS2019 support
 
diff --git a/Engine/Source/Programs/UnrealVS/UnrealVS.2022/Properties/AssemblyInfo.cs b/Engine/Source/Programs/UnrealVS/UnrealVS.2022/Properties/AssemblyInfo.cs
--- a/Engine/Source/Programs/UnrealVS/UnrealVS.2022/Properties/AssemblyInfo.cs
+++ b/Engine/Source/Programs/UnrealVS/UnrealVS.2022/Properties/AssemblyInfo.cs
@@ -31,6 +31,6 @@ using System.Runtime.InteropServices;
 // You can specify all the values or you can default the Revision and Build Numbers 
 // by using the '*' as shown below:
 
-[assembly: AssemblyVersion("1.84.6.0")]
-[assembly: AssemblyFileVersion("1.84.6.0")]
+[assembly: AssemblyVersion("1.84.7.0")]
+[assembly: AssemblyFileVersion("1.84.7.0")]
 
diff --git a/Engine/Source/Programs/UnrealVS/UnrealVS.2022/UnrealVS.2022.csproj b/Engine/Source/Programs/UnrealVS/UnrealVS.2022/UnrealVS.2022.csproj
--- a/Engine/Source/Programs/UnrealVS/UnrealVS.2022/UnrealVS.2022.csproj
+++ b/Engine/Source/Programs/UnrealVS/UnrealVS.2022/UnrealVS.2022.csproj
@@ -153,20 +153,20 @@
     <PackageReference Include="Microsoft.CSharp">
       <Version>4.7.0</Version>
     </PackageReference>
-    <PackageReference Include="Microsoft.VisualStudio.Sdk" Version="17.10.40171">
+    <PackageReference Include="Microsoft.VisualStudio.Sdk" Version="17.8.37222">
     </PackageReference>
     <PackageReference Include="Microsoft.VisualStudio.Threading">
-      <Version>17.10.48</Version>
+      <Version>17.8.14</Version>
     </PackageReference>
     <PackageReference Include="Microsoft.VisualStudio.Threading.Analyzers">
-      <Version>17.10.48</Version>
+      <Version>17.8.14</Version>
       <IncludeAssets>runtime; build; native; contentfiles; analyzers; buildtransitive</IncludeAssets>
       <PrivateAssets>all</PrivateAssets>
     </PackageReference>
     <PackageReference Include="Microsoft.VisualStudio.VCProjectEngine">
-      <Version>17.10.40170</Version>
+      <Version>17.8.37222</Version>
     </PackageReference>
-    <PackageReference Include="Microsoft.VSSDK.BuildTools" Version="17.11.414">
+    <PackageReference Include="Microsoft.VSSDK.BuildTools" Version="17.8.2365">
       <IncludeAssets>runtime; build; native; contentfiles; analyzers; buildtransitive</IncludeAssets>
       <PrivateAssets>all</PrivateAssets>
     </PackageReference>
