[UnrealVs]
* Fixed so ubavisualizer panel works with dpi scaling through remote desktop
* Fixed so when clicking in ubavisualizer panel visual studio moves to top

diff --git a/Engine/Source/Programs/UnrealVS/UnrealVS.Shared/UbaVisualizer.cs b/Engine/Source/Programs/UnrealVS/UnrealVS.Shared/UbaVisualizer.cs
--- a/Engine/Source/Programs/UnrealVS/UnrealVS.Shared/UbaVisualizer.cs
+++ b/Engine/Source/Programs/UnrealVS/UnrealVS.Shared/UbaVisualizer.cs
@@ -9,6 +9,8 @@ using System.Diagnostics;
 using System.Diagnostics.CodeAnalysis;
 using System.Runtime.InteropServices;
 using System.Security;
+using System.Threading.Tasks;
+using System.Windows;
 using System.Windows.Forms;
 using System.Windows.Interop;
 
@@ -99,12 +101,14 @@ namespace UnrealVS
 	internal class UbaVisualizerHost : Hwnd​Host
 	{
 		string VisualizerExe;
+		UbaVisualizerWindow Window;
 		IntPtr HwndVisualizer;
 		WndProc HwndProc;
 		IntPtr HwndStatic;
 
-		public UbaVisualizerHost(string visualizerExe)
+		public UbaVisualizerHost(UbaVisualizerWindow window, string visualizerExe)
 		{
+			Window = window;
 			VisualizerExe = visualizerExe;
 		}
 
@@ -153,6 +157,32 @@ namespace UnrealVS
 			if (HwndVisualizer != IntPtr.Zero && NativeMethods.IsWindow(HwndVisualizer))
 			{
 				hasVisualizer = true;
+
+				if (msg == 0x0445) // Special UbaVisualizer message for mouse focus
+				{
+					//NativeMethods.SetWindowPos(hwnd, (IntPtr)(-1), 0, 0, 0, 0, 0x4000 | 0x0010 | 0x0002 | 0x0001);
+					//NativeMethods.SetWindowPos(1hwnd, (IntPtr)(-2), 0, 0, 0, 0, 0x4000 | 0x0010 | 0x0002 | 0x0001);
+
+					//if (!System.Windows.Application.Current.MainWindow.IsActive)
+					{
+						ThreadHelper.JoinableTaskFactory.Run(() =>
+						{
+							System.Windows.Application.Current.MainWindow.Topmost = true;
+							System.Windows.Application.Current.MainWindow.Topmost = false;
+							//System.Windows.Application.Current.MainWindow.Activate();
+							//var DTE = UnrealVSPackage.Instance.DTE;
+							//DTE.MainWindow.Activate();
+
+							//IVsWindowFrame Frame = (IVsWindowFrame)Window.Frame;
+							//Frame.ShowNoActivate();
+							//Focus();
+							//NativeMethods.SetFocus(HwndVisualizer);
+							return Task.CompletedTask;
+						});
+					}
+					//Focus();
+					//NativeMethods.SetActiveWindow(hwnd);
+				}
 			}
 			else
 			{
@@ -164,12 +194,17 @@ namespace UnrealVS
 					hasVisualizer = true;
 					shouldResize = true;
 				}
+				else if (msg == 0x02E0) // WM_DPICHANGED
+				{
+					shouldResize = true;
+				}
 				else if (HwndVisualizer != IntPtr.Zero)
 				{
 					HwndVisualizer = IntPtr.Zero;
 					int style = 0x50000000; // WS_CHILD | WS_VISIBLE
 					string text = $"Waiting for UbaVisualizer process... (parent hwnd: {hwnd.ToString("X").ToLower()})";
-					HwndStatic = NativeMethods.CreateWindowEx(0, "static", text, style, 0, 0, (int)ActualWidth, (int)ActualHeight, hwnd, IntPtr.Zero, IntPtr.Zero, 0);
+					var td = PresentationSource.FromVisual(this).CompositionTarget.TransformToDevice;
+					HwndStatic = NativeMethods.CreateWindowEx(0, "static", text, style, 0, 0, (int)(ActualWidth * td.M11), (int)(ActualHeight * td.M22), hwnd, IntPtr.Zero, IntPtr.Zero, 0);
 				}
 			}
 
@@ -177,14 +212,24 @@ namespace UnrealVS
 
 			if (shouldResize)
 			{
+				uint flags = 0;
+				IntPtr HwndToResize = IntPtr.Zero;
+
 				if (hasVisualizer)
 				{
-					uint flags = 0x0040 | 0x4000; // SWP_SHOWWINDOW | SWP_ASYNCWINDOWPOS
-					NativeMethods.SetWindowPos(HwndVisualizer, IntPtr.Zero, 0, 0, (int)ActualWidth, (int)ActualHeight, flags);
+					flags = 0x0040 | 0x4000; // SWP_SHOWWINDOW | SWP_ASYNCWINDOWPOS
+					HwndToResize = HwndVisualizer;
 				}
 				else if (HwndStatic != IntPtr.Zero)
 				{
-					NativeMethods.SetWindowPos(HwndStatic, IntPtr.Zero, 0, 0, (int)ActualWidth, (int)ActualHeight, 0x0040);
+					flags = 0x0040;
+					HwndToResize = HwndStatic;
+				}
+
+				if (flags != 0)
+				{
+					var td = PresentationSource.FromVisual(this).CompositionTarget.TransformToDevice;
+					NativeMethods.SetWindowPos(HwndToResize, IntPtr.Zero, 0, 0, (int)(ActualWidth * td.M11), (int)(ActualHeight* td.M22), flags);
 				}
 			}
 			return res;
diff --git a/Engine/Source/Programs/UnrealVS/UnrealVS.Shared/UbaVisualizerWindowControl.xaml.cs b/Engine/Source/Programs/UnrealVS/UnrealVS.Shared/UbaVisualizerWindowControl.xaml.cs
--- a/Engine/Source/Programs/UnrealVS/UnrealVS.Shared/UbaVisualizerWindowControl.xaml.cs
+++ b/Engine/Source/Programs/UnrealVS/UnrealVS.Shared/UbaVisualizerWindowControl.xaml.cs
@@ -11,9 +11,12 @@ namespace UnrealVS
 	{
 		string LastSolutionPath = "Uninitialized";
 		string DefaultText = "Starts when opening solution that contains Engine path";
+		internal UbaVisualizerWindow Window;
 
 		public UbaVisualizerWindowControl(UbaVisualizerWindow window)
 		{
+			Window = window;
+
 			InitializeComponent();
 
 
@@ -29,7 +32,7 @@ namespace UnrealVS
 
 			if (visualizerPath != null)
 			{
-				ControlHostElement.Child = new UbaVisualizerHost(visualizerPath);
+				ControlHostElement.Child = new UbaVisualizerHost(Window, visualizerPath);
 			}
 			else if (ControlHostElement.Child is TextBlock textBlock)
 			{
@@ -68,6 +71,7 @@ namespace UnrealVS
 			if (path == string.Empty)
 			{
 				SetChild(null, DefaultText);
+				return;
 			}
 
 			string engineStr = @"\Engine\";
diff --git a/Engine/Source/Programs/UnrealVS/UnrealVS.Shared/UnrealVSPackage.cs b/Engine/Source/Programs/UnrealVS/UnrealVS.Shared/UnrealVSPackage.cs
--- a/Engine/Source/Programs/UnrealVS/UnrealVS.Shared/UnrealVSPackage.cs
+++ b/Engine/Source/Programs/UnrealVS/UnrealVS.Shared/UnrealVSPackage.cs
@@ -1213,12 +1213,18 @@ namespace UnrealVS
 		[DllImport("user32.dll", EntryPoint = "SetWindowPos", CharSet = CharSet.Unicode)]
 		internal static extern bool SetWindowPos(IntPtr hWnd, IntPtr hWndInsertAfter, int X, int Y, int cx, int cy, uint uFlags);
 
+		[DllImport("user32.dll", EntryPoint = "UpdateWindow", CharSet = CharSet.Unicode)]
+		internal static extern bool UpdateWindow(IntPtr hWnd);
+
 		[DllImport("user32.dll", EntryPoint = "IsWindow", CharSet = CharSet.Unicode)]
 		internal static extern bool IsWindow(IntPtr hWnd);
 
 		[DllImport("user32.dll", EntryPoint = "PostMessageW", CharSet = CharSet.Unicode)]
 		internal static extern bool PostMessageW(IntPtr hWnd, int Msg, IntPtr wParam, IntPtr lParam);
 
+		[DllImport("user32.dll", CharSet = CharSet.Unicode)]
+		internal static extern IntPtr SetFocus(IntPtr hWnd);
+
 		[DllImport("user32.dll")]
 		internal static extern IntPtr LoadCursorW(IntPtr hInstance, IntPtr lpCursorName);
 		
