Use VsShellUtilities.OpenDocument instead of running the File.Open command as it's much faster in large solutions.

#rnx

diff --git a/Engine/Source/Programs/UnrealVS/UnrealVS.Shared/FileBrowserWindowControl.xaml.cs b/Engine/Source/Programs/UnrealVS/UnrealVS.Shared/FileBrowserWindowControl.xaml.cs
--- a/Engine/Source/Programs/UnrealVS/UnrealVS.Shared/FileBrowserWindowControl.xaml.cs
+++ b/Engine/Source/Programs/UnrealVS/UnrealVS.Shared/FileBrowserWindowControl.xaml.cs
@@ -525,12 +525,16 @@ namespace UnrealVS
 					continue;
 				toOpen.Add(item.File);
 			}
+
 			ThreadHelper.JoinableTaskFactory.Run(async delegate
 			{
 				await ThreadHelper.JoinableTaskFactory.SwitchToMainThreadAsync();
 				foreach (var f in toOpen)
-					UnrealVSPackage.Instance.DTE.ExecuteCommand("File.OpenFile", "\"" + f + "\"");
+				{
+					VsShellUtilities.OpenDocument(UnrealVSPackage.Instance, f);
+				}
 			});
+
 			return toOpen.Count != 0;
 		}
 
