[UnrealVS]
* Added null check for PresentationSource

diff --git a/Engine/Source/Programs/UnrealVS/UnrealVS.Shared/UbaVisualizer.cs b/Engine/Source/Programs/UnrealVS/UnrealVS.Shared/UbaVisualizer.cs
--- a/Engine/Source/Programs/UnrealVS/UnrealVS.Shared/UbaVisualizer.cs
+++ b/Engine/Source/Programs/UnrealVS/UnrealVS.Shared/UbaVisualizer.cs
@@ -150,6 +150,17 @@ namespace UnrealVS
 			NativeMethods.DestroyWindow(hwnd.Handle);
 		}
 
+		(double, double) GetDpiScale()
+		{
+			PresentationSource source = PresentationSource.FromVisual(this);
+			if (source == null)
+			{
+				return (1.0, 1.0);
+			}
+			var m = source.CompositionTarget.TransformToDevice;
+			return (m.M11, m.M22);
+		}
+
 		protected IntPtr CustomWndProc(IntPtr hwnd, uint msg, IntPtr wParam, IntPtr lParam)
 		{
 			bool hasVisualizer = false;
@@ -208,8 +219,8 @@ namespace UnrealVS
 					HwndVisualizer = IntPtr.Zero;
 					int style = 0x50000000; // WS_CHILD | WS_VISIBLE
 					string text = $"Waiting for UbaVisualizer process... (parent hwnd: {hwnd.ToString("X").ToLower()})";
-					var td = PresentationSource.FromVisual(this).CompositionTarget.TransformToDevice;
-					HwndStatic = NativeMethods.CreateWindowEx(0, "static", text, style, 0, 0, (int)(ActualWidth * td.M11), (int)(ActualHeight * td.M22), hwnd, IntPtr.Zero, IntPtr.Zero, 0);
+					(double scaleWidth, double scaleHeight) = GetDpiScale();
+					HwndStatic = NativeMethods.CreateWindowEx(0, "static", text, style, 0, 0, (int)(ActualWidth * scaleWidth), (int)(ActualHeight * scaleHeight), hwnd, IntPtr.Zero, IntPtr.Zero, 0);
 				}
 			}
 
@@ -233,8 +244,8 @@ namespace UnrealVS
 
 				if (flags != 0)
 				{
-					var td = PresentationSource.FromVisual(this).CompositionTarget.TransformToDevice;
-					NativeMethods.SetWindowPos(HwndToResize, IntPtr.Zero, 0, 0, (int)(ActualWidth * td.M11), (int)(ActualHeight* td.M22), flags);
+					(double scaleWidth, double scaleHeight) = GetDpiScale();
+					NativeMethods.SetWindowPos(HwndToResize, IntPtr.Zero, 0, 0, (int)(ActualWidth * scaleWidth), (int)(ActualHeight* scaleHeight), flags);
 				}
 			}
 			return res;
