Fix issue where UnrealVS could deadlock. VisualStudio after invoking perforce commands.

- When invoking 'p4vc.bat' calling System.Diagnostics.Process.WaitForExit could deadlock if 'p4vcout.out' exists from a previous run and both stdout and stderror were being redirected at the same time.
- Adding a timeout value to WaitForExit (100ms in this case) and retrying the wait several times avoids the deadlock condition.
- As an additional safeguard I added an additional time out value (10 seconds) that controls how long we should wait for the process to return at all. If the process does not return within the given time frame we terminate it and return/log an error.
-- Note that this process invocation is the time it takes for p4vc.bat to run, not the perforce command itself. 10 seconds is more than reasonable for this.
-- It is probably better to preserve the users VisualStudio session as having to terminate that due to a deadlock is much more frustrating than a perforce command not being invoked.

diff --git a/Engine/Extras/UnrealVS/VS2022/AtomFeed.xml b/Engine/Extras/UnrealVS/VS2022/AtomFeed.xml
--- a/Engine/Extras/UnrealVS/VS2022/AtomFeed.xml
+++ b/Engine/Extras/UnrealVS/VS2022/AtomFeed.xml
@@ -19,7 +19,7 @@
     <content type="application/octet-stream" src="UnrealVS.vsix" />
     <Vsix xmlns:xsd="http://www.w3.org/2001/XMLSchema" xmlns:xsi="http://www.w3.org/2001/XMLSchema-instance" xmlns="http://schemas.microsoft.com/developer/vsx-syndication-schema/2010">
       <Id>ddbf523f-7eb6-4887-bd51-85a714ff87eb</Id>
-      <Version>1.84.8.2</Version>
+      <Version>1.84.8.3</Version>
       <Rating xsi:nil="true" />
       <RatingCount xsi:nil="true" />
       <DownloadCount xsi:nil="true" />
diff --git a/Engine/Source/Programs/UnrealVS/ReleaseNotes.txt b/Engine/Source/Programs/UnrealVS/ReleaseNotes.txt
--- a/Engine/Source/Programs/UnrealVS/ReleaseNotes.txt
+++ b/Engine/Source/Programs/UnrealVS/ReleaseNotes.txt
@@ -1,3 +1,6 @@
+1.84.8.3
+* Bug fix where invoking p4v commands could cause visual studio to hang
+
 1.84.8.2
 * Bug fix where the package info string is incorrect
 
diff --git a/Engine/Source/Programs/UnrealVS/UnrealVS.2022/Properties/AssemblyInfo.cs b/Engine/Source/Programs/UnrealVS/UnrealVS.2022/Properties/AssemblyInfo.cs
--- a/Engine/Source/Programs/UnrealVS/UnrealVS.2022/Properties/AssemblyInfo.cs
+++ b/Engine/Source/Programs/UnrealVS/UnrealVS.2022/Properties/AssemblyInfo.cs
@@ -31,6 +31,6 @@ using System.Runtime.InteropServices;
 // You can specify all the values or you can default the Revision and Build Numbers 
 // by using the '*' as shown below:
 
-[assembly: AssemblyVersion("1.84.8.2")]
-[assembly: AssemblyFileVersion("1.84.8.2")]
+[assembly: AssemblyVersion("1.84.8.3")]
+[assembly: AssemblyFileVersion("1.84.8.3")]
 
diff --git a/Engine/Source/Programs/UnrealVS/UnrealVS.Shared/P4Commands.cs b/Engine/Source/Programs/UnrealVS/UnrealVS.Shared/P4Commands.cs
--- a/Engine/Source/Programs/UnrealVS/UnrealVS.Shared/P4Commands.cs
+++ b/Engine/Source/Programs/UnrealVS/UnrealVS.Shared/P4Commands.cs
@@ -1074,6 +1074,11 @@ namespace UnrealVS
 				P4OutputPane.OutputStringThreadSafe($"Running \"{CmdPath}\" {CommandLine}{Environment.NewLine}");
 			});
 
+			const int PollIntervalMs = 100;
+			const int TimeOutS = 10;
+
+			bool bTimedOut = false;
+
 			StringBuilder StdOutSB = new StringBuilder();
 			StringBuilder StdErrSB = new StringBuilder();
 			DateTime Start = DateTime.Now;
@@ -1108,6 +1113,9 @@ namespace UnrealVS
 				ChildProcessList.Add(ChildProcess);
 			}
 
+			Stopwatch ProcessStartTime = new Stopwatch();
+			ProcessStartTime.Start();
+
 			try
 			{
 				ChildProcess.Start();
@@ -1121,34 +1129,49 @@ namespace UnrealVS
 			ChildProcess.BeginErrorReadLine();
 
 			await ProcessExitTaskSource.Task;
-			ChildProcess.WaitForExit();
+
+			while (!ChildProcess.WaitForExit(PollIntervalMs))
+			{
+				if (ProcessStartTime.Elapsed.Seconds >= TimeOutS)
+				{
+					bTimedOut = true;
+					break;
+				}
+			}
 
 			lock (ChildProcessList)
 			{
 				ChildProcessList.Remove(ChildProcess);
 			}
 
-			string StdOut = StdOutSB.ToString();
-			string StdErr = StdErrSB.ToString();
-
-			await ThreadHelper.JoinableTaskFactory.RunAsync(async () =>
+			if (!bTimedOut)
 			{
-				await ThreadHelper.JoinableTaskFactory.SwitchToMainThreadAsync();
-				if (StdOut.Length > 0 && bOutputStdOut)
-				{
-					P4OutputPane.OutputStringThreadSafe(StdOut);
-				}
-				if (StdErr.Length > 0)
+				string StdOut = StdOutSB.ToString();
+				string StdErr = StdErrSB.ToString();
+
+				await ThreadHelper.JoinableTaskFactory.RunAsync(async () =>
 				{
-					P4OutputPane.OutputStringThreadSafe(StdErr);
-				}
+					await ThreadHelper.JoinableTaskFactory.SwitchToMainThreadAsync();
+					if (StdOut.Length > 0 && bOutputStdOut)
+					{
+						P4OutputPane.OutputStringThreadSafe(StdOut);
+					}
+					if (StdErr.Length > 0)
+					{
+						P4OutputPane.OutputStringThreadSafe(StdErr);
+					}
 
-				TimeSpan Duration = DateTime.Now - Start;
-				P4OutputPane.OutputStringThreadSafe($"complete {Duration.TotalSeconds.ToString()} {Environment.NewLine}");
-			});
+					TimeSpan Duration = DateTime.Now - Start;
+					P4OutputPane.OutputStringThreadSafe($"complete {Duration.TotalSeconds.ToString()} {Environment.NewLine}");
+				});
 
-			bool Result = (ChildProcess.ExitCode == 0 && StdErr.Length == 0);
-			return (Result, StdOut, StdErr);
+				bool Result = (ChildProcess.ExitCode == 0 && StdErr.Length == 0);
+				return (Result, StdOut, StdErr);
+			}
+			else
+			{
+				return (false, "", $"Timed out while trying to run: {CmdPath} {CommandLine}");
+			}
 		}
 	}
 
