Show errors when discovering low level tests
- `BeginErrorReadLine` must be called for the `ErrorDataReceived` callback to be invoked.
- Also logging non-zero exit codes and timeouts. All failures should be visible.

diff --git a/Engine/Source/Programs/UnrealVS/Unreal.TestAdapter/TestDiscoverer.cs b/Engine/Source/Programs/UnrealVS/Unreal.TestAdapter/TestDiscoverer.cs
--- a/Engine/Source/Programs/UnrealVS/Unreal.TestAdapter/TestDiscoverer.cs
+++ b/Engine/Source/Programs/UnrealVS/Unreal.TestAdapter/TestDiscoverer.cs
@@ -45,19 +45,38 @@ namespace Unreal.TestAdapter
 			StringBuilder error = new StringBuilder();
 			process.ErrorDataReceived += new DataReceivedEventHandler((sender, e) =>
 			{
-				error.Append(e.Data);
+				if (e.Data != null)
+				{
+					error.AppendFormat("\n\t{0}", e.Data);
+				}
 			});
 
 			process.Start();
+			process.BeginErrorReadLine();
 			string result = process.StandardOutput.ReadToEnd();
 			bool exited = process.WaitForExit(500);
-			int exitCode = process.ExitCode;
-			if (!exited || exitCode != 0)
+
+			if (error.Length > 0)
 			{
-				logger.SendMessage(TestMessageLevel.Error, "Error " + error);
+				string message = string.Format("Test source \"{0}\" produced errors:\n{{{1}\n}}", source, error);
+				logger.SendMessage(TestMessageLevel.Error, message);
+			}
+
+			if (exited)
+			{
+				if (process.ExitCode != 0)
+				{
+					string message = string.Format("Test source \"{0}\" exited with code {1}.", source, process.ExitCode);
+					logger.SendMessage(TestMessageLevel.Error, message);
+				}
+			}
+			else
+			{
+				string message = string.Format("Test source \"{0}\" failed to exit within a reasonable time.", source, process.ExitCode);
+				logger.SendMessage(TestMessageLevel.Error, message);
 			}
-			return result;
 
+			return result;
 		}
 
 		static void ParseXml(string xml, ITestCaseDiscoverySink discoverySink, string source)
