Don't run AutoSDK when running low level tests
- When running multiple tests they are run in separate processes. If each process runs AutoSDK they serialize behind the single UBT instance behavior and grind the system to a halt requiring a reboot. Setting the `-Multiprocess` command line option skips AutoSDK setup.

diff --git a/Engine/Source/Programs/UnrealVS/Unreal.TestAdapter/TestExecutor.cs b/Engine/Source/Programs/UnrealVS/Unreal.TestAdapter/TestExecutor.cs
--- a/Engine/Source/Programs/UnrealVS/Unreal.TestAdapter/TestExecutor.cs
+++ b/Engine/Source/Programs/UnrealVS/Unreal.TestAdapter/TestExecutor.cs
@@ -48,34 +48,29 @@ namespace Unreal.TestAdapter
 
 		void RunTestCase(TestCase testCase, IFrameworkHandle frameworkHandle, CancellationToken token, bool debug)
 		{
-			if(testCase.GetPropertyValue(TestDiscoverer.DisabledTestProperty, false))
+			if (testCase.GetPropertyValue(TestDiscoverer.DisabledTestProperty, false))
 			{
 				frameworkHandle.RecordResult(new TestResult(testCase) { Outcome = TestOutcome.Skipped });
 				return;
 			}
-			string xml = null;
 
+			string tempFile = Path.Combine(Path.GetTempPath(), Path.GetTempFileName());
+			string args = string.Format("\"{0}\" --reporter unreal -o {1} --extra-args -Multiprocess", testCase.FullyQualifiedName, tempFile);
+
+			string xml = null;
 			if (debug)
 			{
-				string tempFile = Path.Combine(Path.GetTempPath(), Path.GetTempFileName());
-				int processId = frameworkHandle.LaunchProcessWithDebuggerAttached(testCase.Source, Path.GetDirectoryName(testCase.Source), "\"" + testCase.FullyQualifiedName + "\"  --reporter unreal -o " + tempFile, null);
+				string workingDir = Path.GetDirectoryName(testCase.Source);
+				int processId = frameworkHandle.LaunchProcessWithDebuggerAttached(testCase.Source, workingDir, args, null);
 				Process process = Process.GetProcessById(processId);
 				process.WaitForExit();
-				if (File.Exists(tempFile))
-				{
-					xml = File.ReadAllText(tempFile);
-					File.Delete(tempFile);
-				}
 			}
 			else
 			{
-				string tempFile = Path.Combine(Path.GetTempPath(), Path.GetTempFileName());
 				using (Process process = new Process())
 				{
 					process.StartInfo.FileName = testCase.Source;
-
-					//wrap in quotes because test names can have spaces
-					process.StartInfo.Arguments = "\"" + testCase.FullyQualifiedName + "\"  --reporter unreal -o " + tempFile;
+					process.StartInfo.Arguments = args;
 					process.StartInfo.UseShellExecute = false;
 
 					if (token.IsCancellationRequested)
@@ -86,14 +81,15 @@ namespace Unreal.TestAdapter
 					frameworkHandle.RecordStart(testCase);
 					process.Start();
 					process.WaitForExit();
-					if (File.Exists(tempFile))
-					{
-						xml = File.ReadAllText(tempFile);
-						File.Delete(tempFile);
-					}
 				}
 			}
 
+			if (File.Exists(tempFile))
+			{
+				xml = File.ReadAllText(tempFile);
+				File.Delete(tempFile);
+			}
+
 			{
 				var testResult = new TestResult(testCase) { Outcome = TestOutcome.NotFound };
 				ParseResult(xml, testResult);
