Fix hang when debugging multiple tests from Visual Studio
- It does not appear to be valid to run `IFrameworkHandle.LaunchProcessWithDebuggerAttached` in parallel. A second call while the first process is still alive will hang indefinitely, even after the first process exits.
- Tests are now run sequentially when debugging.
- Remove CancellationTokenSource. It's not needed, the boolean is sufficient.

diff --git a/Engine/Source/Programs/UnrealVS/Unreal.TestAdapter/TestExecutor.cs b/Engine/Source/Programs/UnrealVS/Unreal.TestAdapter/TestExecutor.cs
--- a/Engine/Source/Programs/UnrealVS/Unreal.TestAdapter/TestExecutor.cs
+++ b/Engine/Source/Programs/UnrealVS/Unreal.TestAdapter/TestExecutor.cs
@@ -4,7 +4,6 @@ using System;
 using System.Collections.Generic;
 using System.Diagnostics;
 using System.IO;
-using System.Threading;
 using System.Threading.Tasks;
 using System.Xml;
 using Microsoft.VisualStudio.TestPlatform.ObjectModel;
@@ -15,39 +14,44 @@ namespace Unreal.TestAdapter
 	[ExtensionUri("executor://UnrealTestExecutor")]
 	class TestExecutor : ITestExecutor
 	{
-		CancellationTokenSource CancellationTokenSource;
 		bool WasCancelled = false;
 
-		//this does appear ever be called.
-		//meaning that launched process will continue running
 		public void Cancel()
 		{
 			WasCancelled = true;
-			CancellationTokenSource?.Cancel();
+		}
+
+		public void RunTests(IEnumerable<string> sources, IRunContext runContext, IFrameworkHandle frameworkHandle)
+		{
 		}
 
 		public void RunTests(IEnumerable<TestCase> tests, IRunContext runContext, IFrameworkHandle frameworkHandle)
 		{
-			//early out if canceled before even started
-			if (WasCancelled)
-				return;
+			WasCancelled = false;
 
-			using (CancellationTokenSource = new CancellationTokenSource())
+			if (runContext.IsBeingDebugged)
+			{
+				foreach (TestCase testCase in tests)
+				{
+					RunTestCase(testCase, frameworkHandle, true);
+				}
+			}
+			else
 			{
-				var token = CancellationTokenSource.Token;
 				Parallel.ForEach(tests, testCase =>
-					{
-						//if canceled skip but continue the loop to let the other processes to shutdown
-						if (!token.IsCancellationRequested)
-						{
-							RunTestCase(testCase, frameworkHandle, token, runContext.IsBeingDebugged);
-						}
-					});
+				{
+					RunTestCase(testCase, frameworkHandle, false);
+				});
 			}
 		}
 
-		void RunTestCase(TestCase testCase, IFrameworkHandle frameworkHandle, CancellationToken token, bool debug)
+		void RunTestCase(TestCase testCase, IFrameworkHandle frameworkHandle, bool debug)
 		{
+			if (WasCancelled)
+			{
+				return;
+			}
+
 			if (testCase.GetPropertyValue(TestDiscoverer.DisabledTestProperty, false))
 			{
 				frameworkHandle.RecordResult(new TestResult(testCase) { Outcome = TestOutcome.Skipped });
@@ -56,10 +60,11 @@ namespace Unreal.TestAdapter
 
 			string tempFile = Path.Combine(Path.GetTempPath(), Path.GetTempFileName());
 			string args = string.Format("\"{0}\" --reporter unreal -o {1} --extra-args -Multiprocess", testCase.FullyQualifiedName, tempFile);
+			frameworkHandle.RecordStart(testCase);
 
-			string xml = null;
 			if (debug)
 			{
+				// It does not appear to be safe to call LaunchProcessWithDebuggerAttached in parallel
 				string workingDir = Path.GetDirectoryName(testCase.Source);
 				int processId = frameworkHandle.LaunchProcessWithDebuggerAttached(testCase.Source, workingDir, args, null);
 				Process process = Process.GetProcessById(processId);
@@ -72,30 +77,22 @@ namespace Unreal.TestAdapter
 					process.StartInfo.FileName = testCase.Source;
 					process.StartInfo.Arguments = args;
 					process.StartInfo.UseShellExecute = false;
-
-					if (token.IsCancellationRequested)
-					{
-						return;
-					}
-
-					frameworkHandle.RecordStart(testCase);
 					process.Start();
 					process.WaitForExit();
 				}
 			}
 
+			string xmlString = null;
 			if (File.Exists(tempFile))
 			{
-				xml = File.ReadAllText(tempFile);
+				xmlString = File.ReadAllText(tempFile);
 				File.Delete(tempFile);
 			}
 
-			{
-				var testResult = new TestResult(testCase) { Outcome = TestOutcome.NotFound };
-				ParseResult(xml, testResult);
-				frameworkHandle.RecordEnd(testCase, testResult.Outcome);
-				frameworkHandle.RecordResult(testResult);
-			}
+			var testResult = new TestResult(testCase) { Outcome = TestOutcome.NotFound };
+			ParseResult(xmlString, testResult);
+			frameworkHandle.RecordEnd(testCase, testResult.Outcome);
+			frameworkHandle.RecordResult(testResult);
 		}
 
 		void ParseResult(string xmlString, TestResult testResult)
@@ -124,7 +121,6 @@ namespace Unreal.TestAdapter
 					}
 				}
 			}
-
 		}
 
 		void ParseTestCaseResult(XmlReader xml, TestResult testResult)
@@ -159,10 +155,5 @@ namespace Unreal.TestAdapter
 				}
 			}
 		}
-
-		public void RunTests(IEnumerable<string> sources, IRunContext runContext, IFrameworkHandle frameworkHandle)
-		{
-			
-		}
 	}
 }
