[UnrealVS]
* Fixed so visualizer panel works for projects not under Engine folder. Should also work with user generated ue projects

diff --git a/Engine/Extras/UnrealVS/VS2022/AtomFeed.xml b/Engine/Extras/UnrealVS/VS2022/AtomFeed.xml
--- a/Engine/Extras/UnrealVS/VS2022/AtomFeed.xml
+++ b/Engine/Extras/UnrealVS/VS2022/AtomFeed.xml
@@ -8,7 +8,7 @@
     <title type="text">UnrealVS</title>
     <summary type="text">Unreal Engine extension for Visual Studio</summary>
     <published>2022-10-21T10:01:57-07:00</published>
-    <updated>2025-05-02T11:21:31+00:00</updated>
+    <updated>2025-05-08T11:21:31+00:00</updated>
     <author>
       <name>Epic Games, Inc.</name>
     </author>
@@ -19,7 +19,7 @@
     <content type="application/octet-stream" src="UnrealVS.vsix" />
     <Vsix xmlns:xsd="http://www.w3.org/2001/XMLSchema" xmlns:xsi="http://www.w3.org/2001/XMLSchema-instance" xmlns="http://schemas.microsoft.com/developer/vsx-syndication-schema/2010">
       <Id>ddbf523f-7eb6-4887-bd51-85a714ff87eb</Id>
-      <Version>1.84.9.1</Version>
+      <Version>1.84.9.2</Version>
       <Rating xsi:nil="true" />
       <RatingCount xsi:nil="true" />
       <DownloadCount xsi:nil="true" />
diff --git a/Engine/Source/Programs/UnrealVS/ReleaseNotes.txt b/Engine/Source/Programs/UnrealVS/ReleaseNotes.txt
--- a/Engine/Source/Programs/UnrealVS/ReleaseNotes.txt
+++ b/Engine/Source/Programs/UnrealVS/ReleaseNotes.txt
@@ -1,3 +1,6 @@
+1.84.9.2
+* Fix issue where UbaVisualizer will not work on projects outside of Engine hierarchy.
+
 1.84.9.1
 * Fix issue where running UnrealVS.P4->History would occasionally spawn an empty p4 history window.
 
diff --git a/Engine/Source/Programs/UnrealVS/UnrealVS.2022/Properties/AssemblyInfo.cs b/Engine/Source/Programs/UnrealVS/UnrealVS.2022/Properties/AssemblyInfo.cs
--- a/Engine/Source/Programs/UnrealVS/UnrealVS.2022/Properties/AssemblyInfo.cs
+++ b/Engine/Source/Programs/UnrealVS/UnrealVS.2022/Properties/AssemblyInfo.cs
@@ -31,6 +31,6 @@ using System.Runtime.InteropServices;
 // You can specify all the values or you can default the Revision and Build Numbers 
 // by using the '*' as shown below:
 
-[assembly: AssemblyVersion("1.84.9.1")]
-[assembly: AssemblyFileVersion("1.84.9.1")]
+[assembly: AssemblyVersion("1.84.9.2")]
+[assembly: AssemblyFileVersion("1.84.9.2")]
 
diff --git a/Engine/Source/Programs/UnrealVS/UnrealVS.Shared/UbaVisualizerWindowControl.xaml.cs b/Engine/Source/Programs/UnrealVS/UnrealVS.Shared/UbaVisualizerWindowControl.xaml.cs
--- a/Engine/Source/Programs/UnrealVS/UnrealVS.Shared/UbaVisualizerWindowControl.xaml.cs
+++ b/Engine/Source/Programs/UnrealVS/UnrealVS.Shared/UbaVisualizerWindowControl.xaml.cs
@@ -1,10 +1,14 @@
 // Copyright Epic Games, Inc. All Rights Reserved.
 
+using Microsoft.VisualStudio.Shell;
 using System;
 using System.Diagnostics;
 using System.IO;
+using System.Linq.Expressions;
 using System.Windows.Controls;
 
+#pragma warning disable VSTHRD010
+
 namespace UnrealVS
 {
 	public partial class UbaVisualizerWindowControl : UserControl
@@ -54,6 +58,32 @@ namespace UnrealVS
 			}
 		}
 
+		string SearchForVisualizer(string dir)
+		{
+			while (true)
+			{
+				string fullPath = Path.Combine(dir, @"Engine\Binaries\Win64\UnrealBuildAccelerator\x64\UbaVisualizer.exe");
+				try
+				{
+					FileInfo fi = new FileInfo(fullPath);
+					if (fi.Exists)
+					{
+						return fullPath;
+					}
+				}
+				catch
+				{
+				}
+				int lastSlash = dir.LastIndexOf('\\');
+				if (lastSlash == -1)
+				{
+					break;
+				}
+				dir = dir.Substring(0, lastSlash);
+			}
+			return null;
+		}
+
 		public void HandleSolutionChanged()
 		{
 			string path = UnrealVSPackage.Instance.SolutionFilepath;
@@ -74,27 +104,21 @@ namespace UnrealVS
 				return;
 			}
 
-			string engineStr = @"\Engine\";
-			int engineIndex = path.IndexOf(engineStr);
-			if (engineIndex == -1)
+			string fullPath = SearchForVisualizer(Path.GetDirectoryName(path));
+			if (fullPath == null)
 			{
-				path = Path.GetDirectoryName(path);
-				engineIndex = path.Length;
-				path += engineStr;
+				foreach (string projectPath in UnrealVSPackage.Instance.GetLoadedProjectPaths())
+				{
+					if (!projectPath.Contains("\\Engine\\"))
+					{
+						continue;
+					}
+					fullPath = SearchForVisualizer(Path.GetDirectoryName(projectPath));
+					break;
+				}	
 			}
 
-			string fullPath = path.Substring(0, engineIndex + engineStr.Length) + @"Binaries\Win64\UnrealBuildAccelerator\x64\UbaVisualizer.exe";
-			bool exists = false;
-			try
-			{
-				FileInfo fi = new FileInfo(fullPath);
-				exists = fi.Exists;
-			}
-			catch
-			{
-			}
-
-			if (!exists)
+			if (fullPath == null)
 			{
 				SetChild(null, DefaultText);
 				return;
@@ -118,4 +142,6 @@ namespace UnrealVS
 		}
 	}
 
-}
\ No newline at end of file
+}
+
+#pragma warning restore VSTHRD010
