UnrealVS CompileScore changes
- Version 1.84.9.3
- Addressing situations where "Compile and Profile Single File" didn't work correctly with Compile Score.
- Disabling "Compile and Profile Single File" when Compile Score is not installed.
- Adding "UnrealVS.UBT->Build and Profile" project context menu option that uses Compile Score.

#Synced-CL 50914396

diff --git a/Engine/Extras/UnrealVS/VS2022/AtomFeed.xml b/Engine/Extras/UnrealVS/VS2022/AtomFeed.xml
--- a/Engine/Extras/UnrealVS/VS2022/AtomFeed.xml
+++ b/Engine/Extras/UnrealVS/VS2022/AtomFeed.xml
@@ -19,7 +19,7 @@
     <content type="application/octet-stream" src="UnrealVS.vsix" />
     <Vsix xmlns:xsd="http://www.w3.org/2001/XMLSchema" xmlns:xsi="http://www.w3.org/2001/XMLSchema-instance" xmlns="http://schemas.microsoft.com/developer/vsx-syndication-schema/2010">
       <Id>ddbf523f-7eb6-4887-bd51-85a714ff87eb</Id>
-      <Version>1.84.9.2</Version>
+      <Version>1.84.9.3</Version>
       <Rating xsi:nil="true" />
       <RatingCount xsi:nil="true" />
       <DownloadCount xsi:nil="true" />
diff --git a/Engine/Source/Programs/UnrealVS/ReleaseNotes.txt b/Engine/Source/Programs/UnrealVS/ReleaseNotes.txt
--- a/Engine/Source/Programs/UnrealVS/ReleaseNotes.txt
+++ b/Engine/Source/Programs/UnrealVS/ReleaseNotes.txt
@@ -1,3 +1,8 @@
+1.84.9.3
+* Addressing situations where "Compile and Profile Single File" didn't work correctly with Compile Score.
+* Disabling "Compile and Profile Single File" when Compile Score is not installed.
+* Adding "UnrealVS.UBT->Build and Profile" project context menu option that uses Compile Score.
+
 1.84.9.2
 * Fix issue where UbaVisualizer will not work on projects outside of Engine hierarchy.
 
diff --git a/Engine/Source/Programs/UnrealVS/UnrealVS.2022/Properties/AssemblyInfo.cs b/Engine/Source/Programs/UnrealVS/UnrealVS.2022/Properties/AssemblyInfo.cs
--- a/Engine/Source/Programs/UnrealVS/UnrealVS.2022/Properties/AssemblyInfo.cs
+++ b/Engine/Source/Programs/UnrealVS/UnrealVS.2022/Properties/AssemblyInfo.cs
@@ -31,6 +31,6 @@ using System.Runtime.InteropServices;
 // You can specify all the values or you can default the Revision and Build Numbers 
 // by using the '*' as shown below:
 
-[assembly: AssemblyVersion("1.84.9.2")]
-[assembly: AssemblyFileVersion("1.84.9.2")]
+[assembly: AssemblyVersion("1.84.9.3")]
+[assembly: AssemblyFileVersion("1.84.9.3")]
 
diff --git a/Engine/Source/Programs/UnrealVS/UnrealVS.Shared/BuildAndProfile.cs b/Engine/Source/Programs/UnrealVS/UnrealVS.Shared/BuildAndProfile.cs
new file mode 100644
--- /dev/null
+++ b/Engine/Source/Programs/UnrealVS/UnrealVS.Shared/BuildAndProfile.cs
@@ -0,0 +1,68 @@
+// Copyright Epic Games, Inc. All Rights Reserved.
+
+using EnvDTE;
+using Microsoft.VisualStudio.Shell;
+using Microsoft.VisualStudio.Shell.Interop;
+using System;
+using System.Collections.Generic;
+using System.ComponentModel.Design;
+using System.Runtime.InteropServices;
+
+namespace UnrealVS
+{
+	class BuildAndProfile : IDisposable
+	{
+		private readonly int ProjectBuildAndProfileButtonID = 0x1071;
+
+		private UnrealProjectBuilder ProjectBuilder = new UnrealProjectBuilder();
+
+		public BuildAndProfile()
+		{
+			ThreadHelper.ThrowIfNotOnUIThread();
+
+			var BuildAndProfileCommand = new OleMenuCommand(ProjectBuildAndProfileCallback, new CommandID(GuidList.UnrealVSCmdSet, ProjectBuildAndProfileButtonID));
+			BuildAndProfileCommand.BeforeQueryStatus += ProjectBuildAndProfileQueryStatus;
+			UnrealVSPackage.Instance.MenuCommandService.AddCommand(BuildAndProfileCommand);
+		}
+
+		public void Dispose()
+		{
+			ProjectBuilder.Dispose();
+		}
+
+		private void ProjectBuildAndProfileQueryStatus(object Sender, EventArgs Args)
+		{
+			ThreadHelper.ThrowIfNotOnUIThread();
+			((MenuCommand)Sender).Enabled = CompileScoreHelper.IsInstalled();
+		}
+
+		private void ProjectBuildAndProfileCallback(object Sender, EventArgs Args)
+		{
+			ThreadHelper.ThrowIfNotOnUIThread();
+
+			TryProjectBuildAndProfile();
+		}
+
+		private bool TryProjectBuildAndProfile()
+		{
+			ThreadHelper.ThrowIfNotOnUIThread();
+
+			// Get the project clicked in the solution explorer by accessing the current selection and converting to a Project if possible.
+			if (!(Utils.GetCurrentlySelectedProject() is Project SelectedProject))
+			{
+				return false;
+			}
+
+			List<string> ExtraArguments = new List<string>
+			{
+				// When profiling:
+				//   Force disable UBA to force a local build.
+				//   Force use of MSVC since this is specifically for CompileScore.
+				//   Force a unique build environment, might also add more compile options.
+				"-NoUba -NoVfs -Compiler=VisualStudio -ProfilingCompile"
+			};
+
+			return ProjectBuilder.BuildProject(SelectedProject, null, ExtraArguments, true);
+		}
+	}
+}
diff --git a/Engine/Source/Programs/UnrealVS/UnrealVS.Shared/CompileScoreHelper.cs b/Engine/Source/Programs/UnrealVS/UnrealVS.Shared/CompileScoreHelper.cs
new file mode 100644
--- /dev/null
+++ b/Engine/Source/Programs/UnrealVS/UnrealVS.Shared/CompileScoreHelper.cs
@@ -0,0 +1,72 @@
+// Copyright Epic Games, Inc. All Rights Reserved.
+
+using EnvDTE;
+using Microsoft.VisualStudio.Shell;
+using Microsoft.VisualStudio.Shell.Interop;
+using System;
+using System.Linq;
+using System.Windows.Forms;
+
+namespace UnrealVS
+{
+	internal class CompileScoreHelper
+	{
+		/// <summary>
+		/// Checks if the CompileScore plugin's commands have been registered.
+		/// </summary>
+		/// <returns>True if the plugin is installed.</returns>
+		public static bool IsInstalled()
+		{
+			ThreadHelper.ThrowIfNotOnUIThread();
+
+			Command Command = UnrealVSPackage.Instance.DTE.Commands
+				.Cast<Command>()
+				.FirstOrDefault(C =>
+				{
+					ThreadHelper.ThrowIfNotOnUIThread();
+					return C.Name == "CompileScore.StartTrace";
+				});
+
+			return Command != null && Command.IsAvailable;
+		}
+
+		public static void AlertNotInstalled()
+		{
+			MessageBox.Show($"This action requires a CompileScore Visual Studio extension that has CompileScore.StartTrace/CompileScore.StopTrace installed", "UnrealVS - Missing CompileScore extension ", MessageBoxButtons.OK);
+		}
+
+		/// <summary>
+		/// Start a CompileScore trace through the plugin's Visual Studio commands.
+		/// If the plugin is not installed, this will alert the user.
+		/// </summary>
+		/// <returns>True if the trace has been started.</returns>
+		public static bool StartTrace()
+		{
+			if (!IsInstalled())
+			{
+				AlertNotInstalled();
+				return false;
+			}
+
+			ThreadHelper.ThrowIfNotOnUIThread();
+
+			UnrealVSPackage.Instance.DTE.ExecuteCommand("CompileScore.StartTrace");
+
+			IVsOutputWindowPane BuildOutputPane = UnrealVSPackage.Instance.GetOutputPane();
+			UnrealVSPackage.Instance.GetOutputPane()?.OutputStringThreadSafe($"1>------ Started trace for CompileScore ------{Environment.NewLine}");
+
+			return true;
+		}
+
+		/// <summary>
+		/// Stop a CompileScore trace through the plugin's Visual Studio commands.
+		/// </summary>
+		public static void StopTrace()
+		{
+			ThreadHelper.ThrowIfNotOnUIThread();
+			UnrealVSPackage.Instance.DTE.ExecuteCommand("CompileScore.StopTrace");
+
+			UnrealVSPackage.Instance.GetOutputPane()?.OutputStringThreadSafe($"1>------ Stopped trace for CompileScore ------{Environment.NewLine}");
+		}
+	}
+}
diff --git a/Engine/Source/Programs/UnrealVS/UnrealVS.Shared/CompileSingleFile.cs b/Engine/Source/Programs/UnrealVS/UnrealVS.Shared/CompileSingleFile.cs
--- a/Engine/Source/Programs/UnrealVS/UnrealVS.Shared/CompileSingleFile.cs
+++ b/Engine/Source/Programs/UnrealVS/UnrealVS.Shared/CompileSingleFile.cs
@@ -6,12 +6,9 @@ using Microsoft.VisualStudio.Shell.Interop;
 using System;
 using System.Collections.Generic;
 using System.ComponentModel.Design;
-using System.Diagnostics;
 using System.IO;
-using System.Windows.Forms;
-using System.Text;
 using System.Linq;
-using System.Threading.Tasks;
+using System.Windows.Forms;
 
 namespace UnrealVS
 {
@@ -23,68 +20,50 @@ namespace UnrealVS
 		private const int CompileAndProfileSingleFileButtonID = 0x1078;
 		private const int GenerateAssemblyFileButtonID = 0x1079;
 		private const int UBTSubMenuID = 0x3103;
-		private string	  FileToCompileOriginalExt = "";
 
-		static readonly HashSet<string> ValidExtensions = new HashSet<string>(StringComparer.OrdinalIgnoreCase) { ".c", ".cc", ".cpp", ".h", ".cxx", ".ispc" };
+		private static readonly HashSet<string> ValidExtensions = new HashSet<string>(StringComparer.OrdinalIgnoreCase) { ".c", ".cc", ".cpp", ".h", ".cxx", ".ispc" };
 
-		System.Diagnostics.Process ChildProcess;
-		private OleMenuCommand SubMenuCommand;
+		private UnrealProjectBuilder ProjectBuilder = new UnrealProjectBuilder();
 
 		public CompileSingleFile()
 		{
-			CommandID CommandID = new CommandID(GuidList.UnrealVSCmdSet, CompileSingleFileButtonID);
-			var CompileSingleFileButtonCommand = new OleMenuCommand(new EventHandler(CompileSingleFileButtonHandler), CommandID);
-			CompileSingleFileButtonCommand.BeforeQueryStatus += CompileSingleFileButtonCommand_BeforeQueryStatus;
-			UnrealVSPackage.Instance.MenuCommandService.AddCommand(CompileSingleFileButtonCommand);
-
-			CommandID CommandID2 = new CommandID(GuidList.UnrealVSCmdSet, PreprocessSingleFileButtonID);
-			var PreprocessSingleFileButtonCommand = new OleMenuCommand(new EventHandler(CompileSingleFileButtonHandler), CommandID2);
-			PreprocessSingleFileButtonCommand.BeforeQueryStatus += CompileSingleFileButtonCommand_BeforeQueryStatus;
-			UnrealVSPackage.Instance.MenuCommandService.AddCommand(PreprocessSingleFileButtonCommand);
+			AddCommandInternal(CompileSingleFileButtonID);
+			AddCommandInternal(PreprocessSingleFileButtonID);
+			AddCommandInternal(CompileSingleModuleButtonID);
+			AddCommandInternal(CompileAndProfileSingleFileButtonID);
+			AddCommandInternal(GenerateAssemblyFileButtonID);
 
-			CommandID CommandID3 = new CommandID(GuidList.UnrealVSCmdSet, CompileSingleModuleButtonID);
-			MenuCommand CompileSingleModuleButtonCommand = new MenuCommand(new EventHandler(CompileSingleFileButtonHandler), CommandID3);
-			UnrealVSPackage.Instance.MenuCommandService.AddCommand(CompileSingleModuleButtonCommand);
-
-			CommandID CommandID4 = new CommandID(GuidList.UnrealVSCmdSet, CompileAndProfileSingleFileButtonID);
-			var CompileAndProfileSingleFileButtonCommand = new OleMenuCommand(new EventHandler(CompileSingleFileButtonHandler), CommandID4);
-			CompileAndProfileSingleFileButtonCommand.BeforeQueryStatus += CompileSingleFileButtonCommand_BeforeQueryStatus;
-			UnrealVSPackage.Instance.MenuCommandService.AddCommand(CompileAndProfileSingleFileButtonCommand);
+			// add sub menu for UBT commands
+			UnrealVSPackage.Instance.MenuCommandService.AddCommand(new OleMenuCommand(null, new CommandID(GuidList.UnrealVSCmdSet, UBTSubMenuID)));
 
-			CommandID CommandID5 = new CommandID(GuidList.UnrealVSCmdSet, GenerateAssemblyFileButtonID);
-			var GenerateAssemblyFileButtonCommand = new OleMenuCommand(new EventHandler(CompileSingleFileButtonHandler), CommandID5);
-			GenerateAssemblyFileButtonCommand.BeforeQueryStatus += CompileSingleFileButtonCommand_BeforeQueryStatus;
-			UnrealVSPackage.Instance.MenuCommandService.AddCommand(GenerateAssemblyFileButtonCommand);
+			void AddCommandInternal(int CommandID)
+			{
+				var MenuCommand = new OleMenuCommand(CompileSingleFileButtonHandler, new CommandID(GuidList.UnrealVSCmdSet, CommandID));
+				MenuCommand.BeforeQueryStatus += CompileSingleFileButtonCommand_BeforeQueryStatus;
+				UnrealVSPackage.Instance.MenuCommandService.AddCommand(MenuCommand);
+			}
+		}
 
-			// add sub menu for UBT commands
-			SubMenuCommand = new OleMenuCommand(null, new CommandID(GuidList.UnrealVSCmdSet, UBTSubMenuID));
-			UnrealVSPackage.Instance.MenuCommandService.AddCommand(SubMenuCommand);
+		public void Dispose()
+		{
+			ProjectBuilder.Dispose();
 		}
 
-		private void CompileSingleFileButtonCommand_BeforeQueryStatus(object sender, EventArgs e)
+		private void CompileSingleFileButtonCommand_BeforeQueryStatus(object Sender, EventArgs e)
 		{
 			ThreadHelper.ThrowIfNotOnUIThread();
-			DTE DTE = UnrealVSPackage.Instance.DTE;
-			MenuCommand MenuCommand = sender as MenuCommand;
-			if (MenuCommand == null)
-			{
-				return;
-			}
 
-			if (DTE?.ActiveDocument == null)
+			if (Sender is MenuCommand MenuCommand)
 			{
-				MenuCommand.Enabled = false;
-				return;
-			}
-
-			// Check if the requested file is valid
-			string FileToCompileExt = Path.GetExtension(DTE.ActiveDocument.FullName);
-			MenuCommand.Enabled = ValidExtensions.Contains(FileToCompileExt);
-		}
+				DTE DTE = UnrealVSPackage.Instance.DTE;
 
-		public void Dispose()
-		{
-			KillChildProcess();
+				MenuCommand.Enabled =
+					DTE?.ActiveDocument != null
+					// Check if compile score is required and installed
+					&& (MenuCommand.CommandID.ID != CompileAndProfileSingleFileButtonID || CompileScoreHelper.IsInstalled())
+					// Check if the requested file is valid
+					&& ValidExtensions.Contains(Path.GetExtension(DTE.ActiveDocument.FullName));
+			}
 		}
 
 		void CompileSingleFileButtonHandler(object Sender, EventArgs Args)
@@ -100,12 +79,12 @@ namespace UnrealVS
 
 			if (!TryCompileSingleFileOrModule(bIsFile, bPreprocessOnly, bProfile, bGenerateAssembly))
 			{
-				DTE DTE = UnrealVSPackage.Instance.DTE;
-
-				if (bProfile)
-					if (!StartTrace())
-						return;
+				if (bProfile && !CompileScoreHelper.StartTrace())
+				{
+					return;
+				}
 
+				DTE DTE = UnrealVSPackage.Instance.DTE;
 				DTE.ExecuteCommand("Build.Compile");
 
 				if (bProfile)
@@ -119,8 +98,10 @@ namespace UnrealVS
 							{
 								await ThreadHelper.JoinableTaskFactory.SwitchToMainThreadAsync();
 								if (DTE.Solution.SolutionBuild.BuildState != vsBuildState.vsBuildStateDone)
+								{
 									return;
-								StopTrace();
+								}
+								CompileScoreHelper.StopTrace();
 								Done = true;
 
 							});
@@ -132,71 +113,20 @@ namespace UnrealVS
 			}
 		}
 
-		void KillChildProcess()
+		static string FindModuleForFile(string FileName, string RootDirectory)
 		{
-			if (ChildProcess != null)
+			for (string Dir = Path.GetDirectoryName(FileName);
+				Dir != null && !String.Equals(Dir, RootDirectory, StringComparison.OrdinalIgnoreCase);
+				Dir = Path.GetDirectoryName(Dir))
 			{
-				if (!ChildProcess.HasExited)
+				string BuildFile = Directory.EnumerateFiles(Dir, "*.Build.cs", SearchOption.TopDirectoryOnly).FirstOrDefault();
+				if (BuildFile != null)
 				{
-					ChildProcess.Kill();
-					ChildProcess.WaitForExit();
+					// *.Build.cs -> *.Build -> *
+					return Path.GetFileNameWithoutExtension(Path.GetFileNameWithoutExtension(BuildFile));
 				}
-				ChildProcess.Dispose();
-				ChildProcess = null;
-			}
-		}
-
-		string FindModuleForFile(string fileName, string rootDirectory)
-		{
-			string directory = Path.GetDirectoryName(fileName);
-			IEnumerable<string> buildFiles = Directory.EnumerateFiles(directory, "*.build.cs", SearchOption.TopDirectoryOnly);
-			string buildFile = buildFiles.FirstOrDefault();
-			if (buildFile != null)
-			{
-				return Path.GetFileName(buildFile.Substring(0, buildFile.LastIndexOf(".build.cs", StringComparison.OrdinalIgnoreCase)));
-			}
-			if (string.Equals(directory, rootDirectory, StringComparison.OrdinalIgnoreCase))
-			{
-				return null;
-			}
-			return FindModuleForFile(directory, rootDirectory);
-		}
-
-		bool StartTrace()
-		{
-			ThreadHelper.ThrowIfNotOnUIThread();
-			DTE DTE = UnrealVSPackage.Instance.DTE;
-
-			var Commands = DTE.Commands.Cast<Command>();
-#pragma warning disable VSTHRD010
-			Command Command = Commands.FirstOrDefault((C) => C.Name == "CompileScore.StartTrace");
-#pragma warning restore VSTHRD010
-			if (Command == null || !Command.IsAvailable)
-			{
-				MessageBox.Show($"CompileSingleFileAndProfile requires a CompileScore visual studio extension that has CompileScore.StartTrace/CompileScore.StopTrace installed", "UnrealVS - Missing CompileScore extension ", MessageBoxButtons.OK);
-				return false;
-			}
-			Object CustomIn = null;
-			Object CustomOut = null;
-			DTE.Commands.Raise(Command.Guid, Command.ID, ref CustomIn, ref CustomOut);
-			//DTE.ExecuteCommand("CompileScore.StartTrace");
-			IVsOutputWindowPane BuildOutputPane = UnrealVSPackage.Instance.GetOutputPane();
-			if (BuildOutputPane == null)
-			{
-				BuildOutputPane.OutputStringThreadSafe($"1>------ Started trace for CompileScore ------{Environment.NewLine}");
-			}
-			return true;
-		}
-
-		void StopTrace()
-		{
-			ThreadHelper.ThrowIfNotOnUIThread();
-			UnrealVSPackage.Instance.DTE.ExecuteCommand("CompileScore.StopTrace");
-			IVsOutputWindowPane BuildOutputPane = UnrealVSPackage.Instance.GetOutputPane();
-			if (BuildOutputPane != null)
-			{
-				BuildOutputPane.OutputStringThreadSafe($"1>------ Stopped trace for CompileScore ------{Environment.NewLine}");
 			}
+			return null;
 		}
 
 		bool TryCompileSingleFileOrModule(bool bIsFile, bool bPreProcessOnly, bool bProfile, bool bGenerateAssembly)
@@ -204,29 +134,6 @@ namespace UnrealVS
 			ThreadHelper.ThrowIfNotOnUIThread();
 			DTE DTE = UnrealVSPackage.Instance.DTE;
 
-			// Activate the output window
-			Window Window = DTE.Windows.Item(EnvDTE.Constants.vsWindowKindOutput);
-			Window.Activate();
-
-			// Find or create the 'Build' window
-			IVsOutputWindowPane BuildOutputPane = UnrealVSPackage.Instance.GetOutputPane();
-			if (BuildOutputPane == null)
-			{
-				Logging.WriteLine("CompileSingleFile: Build Output Pane not found");
-				return false;
-			}
-
-			// If there's already a build in progress, offer to cancel it
-			if (ChildProcess != null && !ChildProcess.HasExited)
-			{
-				if (MessageBox.Show("Cancel current compile?", "Compile in progress", MessageBoxButtons.YesNo) == DialogResult.Yes)
-				{
-					KillChildProcess();
-					BuildOutputPane.OutputStringThreadSafe($"1>  Build cancelled.{Environment.NewLine}");
-				}
-				return true;
-			}
-
 			// Check we've got a file open
 			if (DTE.ActiveDocument == null)
 			{
@@ -241,57 +148,29 @@ namespace UnrealVS
 				Logging.WriteLine("CompileSingleFile: ProjectHierarchy not found");
 				return false;
 			}
-			Project StartupProject = Utils.HierarchyObjectToProject(ProjectHierarchy);
-			if (StartupProject == null)
-			{
-				Logging.WriteLine("CompileSingleFile: StartupProject not found");
-				return false;
-			}
-			if (!(StartupProject.Object is Microsoft.VisualStudio.VCProjectEngine.VCProject VCStartupProject))
-			{
-				Logging.WriteLine("CompileSingleFile: VCStartupProject not found");
-				return false;
-			}
 
-			// Get the active configuration for the startup project
-			Configuration ActiveConfiguration = StartupProject.ConfigurationManager.ActiveConfiguration;
-			string ActiveConfigurationName = $"{ActiveConfiguration.ConfigurationName}|{ActiveConfiguration.PlatformName}";
-			Microsoft.VisualStudio.VCProjectEngine.VCConfiguration ActiveVCConfiguration = (VCStartupProject.Configurations as Microsoft.VisualStudio.VCProjectEngine.IVCCollection).Item(ActiveConfigurationName) as Microsoft.VisualStudio.VCProjectEngine.VCConfiguration;
-			if (ActiveVCConfiguration == null)
+			if (!(Utils.HierarchyObjectToProject(ProjectHierarchy) is Project StartupProject))
 			{
-				Logging.WriteLine("CompileSingleFile: VCStartupProject ActiveConfiguration not found");
-				return false;
-			}
-
-			// Get the NMake settings for this configuration
-			Microsoft.VisualStudio.VCProjectEngine.VCNMakeTool ActiveNMakeTool = (ActiveVCConfiguration.Tools as Microsoft.VisualStudio.VCProjectEngine.IVCCollection).Item("VCNMakeTool") as Microsoft.VisualStudio.VCProjectEngine.VCNMakeTool;
-			if (ActiveNMakeTool == null)
-			{
-				MessageBox.Show($"No NMakeTool set for Project {VCStartupProject.Name} set for single-file compile.", "UnrealVS - NMakeTool not set", MessageBoxButtons.OK);
+				Logging.WriteLine("CompileSingleFile: StartupProject not found");
 				return false;
 			}
 
-			// Save all the open documents
-			DTE.ExecuteCommand("File.SaveAll");
-
 			// Check if the requested file is valid
 			string FileToCompile = DTE.ActiveDocument.FullName;
 			string FileToCompileExt = Path.GetExtension(FileToCompile);
 
 			string CompilingText;
-			List<string> UBTArguments = new List<string>
-			{
-				"-WorkingDir=\"$(MSBuildProjectDirectory)\"",
-			};
+			List<string> ExtraArguments = new List<string>();
+
 			if (bPreProcessOnly)
 			{
-				UBTArguments.Add("-NoXGE -NoSNDBS -NoFASTBuild");
-				UBTArguments.Add("-Preprocess");
+				ExtraArguments.Add("-NoXGE -NoSNDBS -NoFASTBuild");
+				ExtraArguments.Add("-Preprocess");
 			}
 			else if (bGenerateAssembly)
 			{
-				UBTArguments.Add("-NoXGE -NoSNDBS -NoFASTBuild");
-				UBTArguments.Add("-WithAssembly");
+				ExtraArguments.Add("-NoXGE -NoSNDBS -NoFASTBuild");
+				ExtraArguments.Add("-WithAssembly");
 			}
 
 			if (bIsFile)
@@ -303,7 +182,7 @@ namespace UnrealVS
 				}
 
 				CompilingText = FileToCompile;
-				UBTArguments.Add($"-SingleFile=\"{FileToCompile}\"");
+				ExtraArguments.Add($"-SingleFile=\"{FileToCompile}\"");
 			}
 			else
 			{
@@ -314,153 +193,10 @@ namespace UnrealVS
 					return true;
 				}
 				CompilingText = ModuleName;
-				UBTArguments.Add($"-Module=\"{ModuleName}\"");
-			}
-
-			// If there's already a build in progress, don't let another one start
-			if (DTE.Solution.SolutionBuild.BuildState == vsBuildState.vsBuildStateInProgress)
-			{
-				if (MessageBox.Show("Cancel current compile?", "Compile in progress", MessageBoxButtons.YesNo) == DialogResult.Yes)
-				{
-					DTE.ExecuteCommand("Build.Cancel");
-				}
-				return true;
-			}
-
-			// Make sure any existing build is stopped
-			KillChildProcess();
-
-			// Set up the output pane
-			BuildOutputPane.Activate();
-			BuildOutputPane.Clear();
-
-			// Set up event handlers 
-			DTE.Events.BuildEvents.OnBuildBegin += BuildEvents_OnBuildBegin;
-
-			// Create a delegate for handling output messages
-			List<string> PreprocessedFiles = new List<string>();
-			List<string> AssemblyFiles = new List<string>();
-			void OutputHandler(object Sender, DataReceivedEventArgs Args) 
-			{ 
-				if (Args.Data != null) 
-				{ 
-					if (Args.Data.Contains("PreProcessPath:"))
-					{
-						PreprocessedFiles.Add(Args.Data.Replace("PreProcessPath:", "").Trim());
-					}
-					else if (Args.Data.Contains("AssemblyPath:"))
-					{
-						AssemblyFiles.Add(Args.Data.Replace("AssemblyPath:", "").Trim());
-					}
-					else
-					{
-						BuildOutputPane.OutputStringThreadSafe($"1>  {Args.Data}{Environment.NewLine}");
-					}
-				} 
-			}
-
-			if (bProfile)
-				if (!StartTrace())
-					return true;
-
-			BuildOutputPane.OutputStringThreadSafe($"1>------ Build started: Project: {StartupProject.Name}, Configuration: {ActiveConfiguration.ConfigurationName} {ActiveConfiguration.PlatformName} ------{Environment.NewLine}");
-			BuildOutputPane.OutputStringThreadSafe($"1>  Compiling {CompilingText}{Environment.NewLine}");
-
-			string SolutionDir = Path.GetDirectoryName(UnrealVSPackage.Instance.SolutionFilepath).TrimEnd(Path.DirectorySeparatorChar) + Path.DirectorySeparatorChar;
-			// Get the build command line and escape any environment variables that we use
-			string BuildCommandLine = ActiveVCConfiguration.Evaluate(ActiveNMakeTool.BuildCommandLine);
-			string UBTArgument = ActiveVCConfiguration.Evaluate(string.Join(" ", UBTArguments));
-			string WorkingDirectory = ActiveVCConfiguration.Evaluate("$(MSBuildProjectDirectory)");
-
-			FileToCompileOriginalExt = FileToCompileExt;
-
-			// Spawn the new process
-			ChildProcess = new System.Diagnostics.Process();
-			ChildProcess.StartInfo.FileName = Path.Combine(Environment.SystemDirectory, "cmd.exe");
-			ChildProcess.StartInfo.Arguments = $"/C \"{BuildCommandLine} {UBTArgument}\"";
-			ChildProcess.StartInfo.WorkingDirectory = WorkingDirectory;
-			ChildProcess.StartInfo.UseShellExecute = false;
-			ChildProcess.StartInfo.RedirectStandardOutput = true;
-			ChildProcess.StartInfo.RedirectStandardError = true;
-			ChildProcess.StartInfo.CreateNoWindow = true;
-			ChildProcess.OutputDataReceived += OutputHandler;
-			ChildProcess.ErrorDataReceived += OutputHandler;
-			if (bPreProcessOnly || bGenerateAssembly  || bProfile)
-			{
-				// add an event handler to respond to the exit of the preprocess request
-				// and open the generated file if it exists.
-				ChildProcess.EnableRaisingEvents = true;
-				ChildProcess.Exited += new EventHandler((s, e) => UbtProcessExitHandler(PreprocessedFiles, AssemblyFiles, bProfile));
-			}
-
-			ChildProcess.Start();
-			ChildProcess.BeginOutputReadLine();
-			ChildProcess.BeginErrorReadLine();
-
-			return true;
-		}
-
-		private void UbtProcessExitHandler(IEnumerable<string> PreprocessedFiles, IEnumerable<string> AssemblyFiles, bool bIsProfiling)
-		{
-			if (bIsProfiling)
-			{
-				ThreadHelper.JoinableTaskFactory.Run(async () =>
-				{
-					await ThreadHelper.JoinableTaskFactory.SwitchToMainThreadAsync();
-					StopTrace();
-				});
+				ExtraArguments.Add($"-Module=\"{ModuleName}\"");
 			}
 
-			// not all compile actions support pre-process - check it exists
-			foreach (string PreprocessedFile in PreprocessedFiles)
-			{
-				ThreadHelper.JoinableTaskFactory.Run(async () =>
-				{
-					await ThreadHelper.JoinableTaskFactory.SwitchToMainThreadAsync();
-					OpenPreprocessedFile(PreprocessedFile);
-				});
-			}
-			foreach (string AssemblyFile in AssemblyFiles)
-			{
-				ThreadHelper.JoinableTaskFactory.Run(async () =>
-				{
-					await ThreadHelper.JoinableTaskFactory.SwitchToMainThreadAsync();
-					OpenAssemblyFile(AssemblyFile);
-				});
-			}
-		}
-
-		private void OpenPreprocessedFile(string PPFullPath)
-		{
-			ThreadHelper.ThrowIfNotOnUIThread();
-
-			if (File.Exists(PPFullPath))
-			{
-				// if the file exists, rename it to isolate the file and have its extension be the original to maintain syntax highlighting
-				string Dir = Path.GetDirectoryName(PPFullPath);
-				string FileName = Path.GetFileNameWithoutExtension(Path.GetFileNameWithoutExtension(PPFullPath)) + "_preprocessed";
-
-				string RenamedFile = Path.Combine(Dir, FileName) + FileToCompileOriginalExt;
-
-				File.Copy(PPFullPath, RenamedFile, true /*overwrite*/);
-
-				UnrealVSPackage.Instance.DTE.ExecuteCommand("File.OpenFile", $"\"{RenamedFile}\"");
-			}
-		}
-
-		private void OpenAssemblyFile(string AsmFullPath)
-		{
-			ThreadHelper.ThrowIfNotOnUIThread();
-
-			if (File.Exists(AsmFullPath))
-			{
-				UnrealVSPackage.Instance.DTE.ExecuteCommand("File.OpenFile", $"\"{AsmFullPath}\"");
-			}
-		}
-
-		private void BuildEvents_OnBuildBegin(vsBuildScope Scope, vsBuildAction Action)
-		{
-			KillChildProcess();
+			return ProjectBuilder.BuildProject(StartupProject, CompilingText, ExtraArguments, bProfile);
 		}
 	}
 }
diff --git a/Engine/Source/Programs/UnrealVS/UnrealVS.Shared/QuickBuild.cs b/Engine/Source/Programs/UnrealVS/UnrealVS.Shared/QuickBuild.cs
--- a/Engine/Source/Programs/UnrealVS/UnrealVS.Shared/QuickBuild.cs
+++ b/Engine/Source/Programs/UnrealVS/UnrealVS.Shared/QuickBuild.cs
@@ -102,22 +102,17 @@ namespace UnrealVS
 		/// <summary>
 		/// The invoke handler passed to this object's QuickBuildMenuCommand, called by the base OleMenuCommand when an item is clicked
 		/// </summary>
-		private void OnInvokedDynamicItem(object sender, EventArgs args)
+		private void OnInvokedDynamicItem(object Sender, EventArgs args)
 		{
 			ThreadHelper.ThrowIfNotOnUIThread();
-			var MenuCommand = (QuickBuildMenuCommand)sender;
+			var MenuCommand = (QuickBuildMenuCommand)Sender;
 
 			// Get the project clicked in the solution explorer by accessing the current selection and converting to a Project if possible.
-			UnrealVSPackage.Instance.SelectionManager.GetCurrentSelection(out IntPtr HierarchyPtr, out uint ProjectItemId, out IVsMultiItemSelect MultiItemSelect, out IntPtr SelectionContainerPtr);
-			if (HierarchyPtr == IntPtr.Zero) return;
-
-			if (!(Marshal.GetTypedObjectForIUnknown(HierarchyPtr, typeof(IVsHierarchy)) is IVsHierarchy SelectedHierarchy)) return;
-
-			var SelectedProject = Utils.HierarchyObjectToProject(SelectedHierarchy);
-			if (SelectedProject == null) return;
-
-			// Builds the selected project with the clicked platform and config
-			Utils.ExecuteProjectBuild(SelectedProject, ConfigurationName, MenuCommand.Text, BatchBuilderToolControl.BuildJob.BuildJobType.Build, null, null);
+			if (Utils.GetCurrentlySelectedProject() is Project SelectedProject)
+			{
+				// Builds the selected project with the clicked platform and config
+				Utils.ExecuteProjectBuild(SelectedProject, ConfigurationName, MenuCommand.Text, BatchBuilderToolControl.BuildJob.BuildJobType.Build, null, null);
+			}
 		}
 
 		/// <summary>
diff --git a/Engine/Source/Programs/UnrealVS/UnrealVS.Shared/UnrealProjectBuilder.cs b/Engine/Source/Programs/UnrealVS/UnrealVS.Shared/UnrealProjectBuilder.cs
new file mode 100644
--- /dev/null
+++ b/Engine/Source/Programs/UnrealVS/UnrealVS.Shared/UnrealProjectBuilder.cs
@@ -0,0 +1,253 @@
+// Copyright Epic Games, Inc. All Rights Reserved.
+
+using EnvDTE;
+using Microsoft.VisualStudio.Shell;
+using Microsoft.VisualStudio.Shell.Interop;
+using Microsoft.VisualStudio.VCProjectEngine;
+using System;
+using System.Collections.Generic;
+using System.Diagnostics;
+using System.IO;
+using System.Windows.Forms;
+
+namespace UnrealVS
+{
+	class UnrealProjectBuilder : IDisposable
+	{
+		System.Diagnostics.Process ChildProcess;
+		private string FileToCompileOriginalExt = "";
+
+		public UnrealProjectBuilder()
+		{
+		}
+
+		public void Dispose()
+		{
+			KillChildProcess();
+		}
+
+		public bool BuildProject(Project TargetProject, string SpecificTarget, List<string> ExtraArguments, bool bRunningCompileScore)
+		{
+			ThreadHelper.ThrowIfNotOnUIThread();
+			DTE DTE = UnrealVSPackage.Instance.DTE;
+
+			if (!(TargetProject.Object is VCProject VCProject))
+			{
+				Logging.WriteLine("UnrealProjectBuilder: VCProject not found");
+				return false;
+			}
+
+			// Activate the output window
+			Window Window = DTE.Windows.Item(EnvDTE.Constants.vsWindowKindOutput);
+			Window.Activate();
+
+			// Find or create the 'Build' window
+			if (!(UnrealVSPackage.Instance.GetOutputPane() is IVsOutputWindowPane BuildOutputPane))
+			{
+				Logging.WriteLine("UnrealProjectBuilder: Build Output Pane not found");
+				return false;
+			}
+
+			// If there's already a build in progress, offer to cancel it
+			if (ChildProcess != null && !ChildProcess.HasExited)
+			{
+				if (MessageBox.Show("Cancel current compile?", "Compile in progress", MessageBoxButtons.YesNo) == DialogResult.Yes)
+				{
+					KillChildProcess();
+					BuildOutputPane.OutputStringThreadSafe($"1>  Build cancelled.{Environment.NewLine}");
+				}
+				return true;
+			}
+
+			Configuration ActiveConfiguration = TargetProject.ConfigurationManager.ActiveConfiguration;
+			string ActiveConfigurationName = $"{ActiveConfiguration.ConfigurationName}|{ActiveConfiguration.PlatformName}";
+			if (!((VCProject.Configurations as IVCCollection).Item(ActiveConfigurationName) is VCConfiguration ActiveVCConfiguration))
+			{
+				Logging.WriteLine("UnrealProjectBuilder: Project ActiveConfiguration not found");
+				return false;
+			}
+
+			// Get the NMake settings for this configuration
+			if (!((ActiveVCConfiguration.Tools as IVCCollection).Item("VCNMakeTool") is VCNMakeTool ActiveNMakeTool))
+			{
+				MessageBox.Show($"No NMakeTool set for Project {VCProject.Name} set for single-file compile.", "UnrealVS - NMakeTool not set", MessageBoxButtons.OK);
+				return false;
+			}
+
+			// Save all the open documents
+			DTE.ExecuteCommand("File.SaveAll");
+
+			// If there's already a build in progress, don't let another one start
+			if (DTE.Solution.SolutionBuild.BuildState == vsBuildState.vsBuildStateInProgress)
+			{
+				if (MessageBox.Show("Cancel current compile?", "Compile in progress", MessageBoxButtons.YesNo) == DialogResult.Yes)
+				{
+					DTE.ExecuteCommand("Build.Cancel");
+				}
+				return true;
+			}
+
+			// Make sure any existing build is stopped
+			KillChildProcess();
+
+			if (DTE.ActiveDocument != null)
+			{
+				FileToCompileOriginalExt = Path.GetExtension(DTE.ActiveDocument.FullName);
+			}
+
+			List<string> UBTArguments = new List<string>
+			{
+				"-WorkingDir=\"$(MSBuildProjectDirectory)\"",
+			};
+
+			UBTArguments.AddRange(ExtraArguments);
+
+			// Set up the output pane
+			BuildOutputPane.Activate();
+			BuildOutputPane.Clear();
+
+			// Set up event handlers 
+			DTE.Events.BuildEvents.OnBuildBegin += BuildEvents_OnBuildBegin;
+
+			// Create a delegate for handling output messages
+			List<string> PreprocessedFiles = new List<string>();
+			List<string> AssemblyFiles = new List<string>();
+			void OutputHandler(object Sender, DataReceivedEventArgs Args)
+			{
+				if (Args.Data != null)
+				{
+					if (Args.Data.Contains("PreProcessPath:"))
+					{
+						PreprocessedFiles.Add(Args.Data.Replace("PreProcessPath:", "").Trim());
+					}
+					else if (Args.Data.Contains("AssemblyPath:"))
+					{
+						AssemblyFiles.Add(Args.Data.Replace("AssemblyPath:", "").Trim());
+					}
+					else
+					{
+						BuildOutputPane.OutputStringThreadSafe($"1>  {Args.Data}{Environment.NewLine}");
+					}
+				}
+			}
+
+			if (bRunningCompileScore && !CompileScoreHelper.StartTrace())
+			{
+				return false;
+			}
+
+			BuildOutputPane.OutputStringThreadSafe($"1>------ Build started: Project: {TargetProject.Name}, Configuration: {ActiveConfiguration.ConfigurationName} {ActiveConfiguration.PlatformName} ------{Environment.NewLine}");
+			if (!string.IsNullOrEmpty(SpecificTarget))
+			{
+				BuildOutputPane.OutputStringThreadSafe($"1>  Compiling {SpecificTarget}{Environment.NewLine}");
+			}
+
+			string SolutionDir = Path.GetDirectoryName(UnrealVSPackage.Instance.SolutionFilepath).TrimEnd(Path.DirectorySeparatorChar) + Path.DirectorySeparatorChar;
+			// Get the build command line and escape any environment variables that we use
+			string BuildCommandLine = ActiveVCConfiguration.Evaluate(ActiveNMakeTool.BuildCommandLine);
+			string UBTArgument = ActiveVCConfiguration.Evaluate(string.Join(" ", UBTArguments));
+			string WorkingDirectory = ActiveVCConfiguration.Evaluate("$(MSBuildProjectDirectory)");
+
+			// Spawn the new process
+			ChildProcess = new System.Diagnostics.Process();
+			ChildProcess.StartInfo.FileName = Path.Combine(Environment.SystemDirectory, "cmd.exe");
+			ChildProcess.StartInfo.Arguments = $"/C \"{BuildCommandLine} {UBTArgument}\"";
+			ChildProcess.StartInfo.WorkingDirectory = WorkingDirectory;
+			ChildProcess.StartInfo.UseShellExecute = false;
+			ChildProcess.StartInfo.RedirectStandardOutput = true;
+			ChildProcess.StartInfo.RedirectStandardError = true;
+			ChildProcess.StartInfo.CreateNoWindow = true;
+			ChildProcess.OutputDataReceived += OutputHandler;
+			ChildProcess.ErrorDataReceived += OutputHandler;
+			{
+				// add an event handler to respond to the exit of the request
+				// and open the generated file if it exists.
+				ChildProcess.EnableRaisingEvents = true;
+				ChildProcess.Exited += new EventHandler((s, e) => UbtProcessExitHandler(PreprocessedFiles, AssemblyFiles, bRunningCompileScore));
+			}
+
+			ChildProcess.Start();
+			ChildProcess.BeginOutputReadLine();
+			ChildProcess.BeginErrorReadLine();
+
+			return true;
+		}
+
+		void KillChildProcess()
+		{
+			if (ChildProcess != null)
+			{
+				if (!ChildProcess.HasExited)
+				{
+					ChildProcess.Kill();
+					ChildProcess.WaitForExit();
+				}
+				ChildProcess.Dispose();
+				ChildProcess = null;
+			}
+		}
+
+		private void UbtProcessExitHandler(IEnumerable<string> PreprocessedFiles, IEnumerable<string> AssemblyFiles, bool bRunningCompileScore)
+		{
+			if (bRunningCompileScore)
+			{
+				ThreadHelper.JoinableTaskFactory.Run(async () =>
+				{
+					await ThreadHelper.JoinableTaskFactory.SwitchToMainThreadAsync();
+					CompileScoreHelper.StopTrace();
+				});
+			}
+
+			// not all compile actions support pre-process - check it exists
+			foreach (string PreprocessedFile in PreprocessedFiles)
+			{
+				ThreadHelper.JoinableTaskFactory.Run(async () =>
+				{
+					await ThreadHelper.JoinableTaskFactory.SwitchToMainThreadAsync();
+					OpenPreprocessedFile(PreprocessedFile);
+				});
+			}
+			foreach (string AssemblyFile in AssemblyFiles)
+			{
+				ThreadHelper.JoinableTaskFactory.Run(async () =>
+				{
+					await ThreadHelper.JoinableTaskFactory.SwitchToMainThreadAsync();
+					OpenAssemblyFile(AssemblyFile);
+				});
+			}
+		}
+
+		private void OpenPreprocessedFile(string PPFullPath)
+		{
+			ThreadHelper.ThrowIfNotOnUIThread();
+
+			if (File.Exists(PPFullPath))
+			{
+				// if the file exists, rename it to isolate the file and have its extension be the original to maintain syntax highlighting
+				string Dir = Path.GetDirectoryName(PPFullPath);
+				string FileName = Path.GetFileNameWithoutExtension(Path.GetFileNameWithoutExtension(PPFullPath)) + "_preprocessed";
+
+				string RenamedFile = Path.Combine(Dir, FileName) + FileToCompileOriginalExt;
+
+				File.Copy(PPFullPath, RenamedFile, true /*overwrite*/);
+
+				UnrealVSPackage.Instance.DTE.ExecuteCommand("File.OpenFile", $"\"{RenamedFile}\"");
+			}
+		}
+
+		private void OpenAssemblyFile(string AsmFullPath)
+		{
+			ThreadHelper.ThrowIfNotOnUIThread();
+
+			if (File.Exists(AsmFullPath))
+			{
+				UnrealVSPackage.Instance.DTE.ExecuteCommand("File.OpenFile", $"\"{AsmFullPath}\"");
+			}
+		}
+
+		private void BuildEvents_OnBuildBegin(vsBuildScope Scope, vsBuildAction Action)
+		{
+			KillChildProcess();
+		}
+	}
+}
diff --git a/Engine/Source/Programs/UnrealVS/UnrealVS.Shared/UnrealVSPackage.cs b/Engine/Source/Programs/UnrealVS/UnrealVS.Shared/UnrealVSPackage.cs
--- a/Engine/Source/Programs/UnrealVS/UnrealVS.Shared/UnrealVSPackage.cs
+++ b/Engine/Source/Programs/UnrealVS/UnrealVS.Shared/UnrealVSPackage.cs
@@ -284,6 +284,8 @@ namespace UnrealVS
 			// Create the project menu quick builder
 			QuickBuilder = new QuickBuild();
 
+			BuildAndProfile = new BuildAndProfile();
+
 			// Create 'Perforce menu' instance
 			P4CommandsGroup = new P4Commands();
 
@@ -445,6 +447,8 @@ namespace UnrealVS
 			StartupProjectSelector = null;
 			BatchBuilder = null;
 			QuickBuilder = null;
+			BuildAndProfile?.Dispose();
+			BuildAndProfile = null;
 
 			CompileSingleFile?.Dispose();
 			CompileSingleFile = null;
@@ -1109,6 +1113,9 @@ namespace UnrealVS
 		/// CompileSingleFile feature
 		private CompileSingleFile CompileSingleFile;
 
+		/// Our Build and Profile component
+		private BuildAndProfile BuildAndProfile;
+
 		/// Project file generator button
 		private GenerateProjectFiles GenerateProjectFiles;
 
diff --git a/Engine/Source/Programs/UnrealVS/UnrealVS.Shared/Utils.cs b/Engine/Source/Programs/UnrealVS/UnrealVS.Shared/Utils.cs
--- a/Engine/Source/Programs/UnrealVS/UnrealVS.Shared/Utils.cs
+++ b/Engine/Source/Programs/UnrealVS/UnrealVS.Shared/Utils.cs
@@ -10,6 +10,7 @@ using System.Collections.Generic;
 using System.ComponentModel;
 using System.IO;
 using System.Linq;
+using System.Runtime.InteropServices;
 using System.Windows.Controls;
 using System.Windows.Media;
 using System.Windows.Media.Animation;
@@ -64,6 +65,46 @@ namespace UnrealVS
 			return (Project)ProjectObject;
 		}
 
+		/// <summary>
+		/// Get the project selected in the solution explorer by accessing the current selection and converting to a Project if possible.
+		/// </summary>
+		/// <returns></returns>
+		public static Project GetCurrentlySelectedProject()
+		{
+			ThreadHelper.ThrowIfNotOnUIThread();
+
+			IntPtr HierarchyPtr = IntPtr.Zero;
+			IntPtr SelectionContainerPtr = IntPtr.Zero;
+
+			try
+			{
+				UnrealVSPackage.Instance.SelectionManager.GetCurrentSelection(out HierarchyPtr, out uint ProjectItemId, out IVsMultiItemSelect MultiItemSelect, out SelectionContainerPtr);
+
+				if (HierarchyPtr == IntPtr.Zero)
+				{
+					return null;
+				}
+
+				if (!(Marshal.GetTypedObjectForIUnknown(HierarchyPtr, typeof(IVsHierarchy)) is IVsHierarchy SelectedHierarchy))
+				{
+					return null;
+				}
+
+				return HierarchyObjectToProject(SelectedHierarchy);
+			}
+			finally
+			{
+				if (HierarchyPtr != IntPtr.Zero)
+				{
+					Marshal.Release(HierarchyPtr);
+				}
+				if (SelectionContainerPtr != IntPtr.Zero)
+				{
+					Marshal.Release(SelectionContainerPtr);
+				}
+			}
+		}
+
 		/// <summary>
 		/// Converts an IVsHierarchy object to a config provider interface
 		/// </summary>
