UnrealVS Module Building
- Version 1.84.9.4
- Adding "UnrealVS.UBT" to solution folder context menus.
- Adding options to build Unreal modules normally and non-unity. Also adding CompileScore profiling for both options.
- Only copy the Release vsix to the Extras folder.

#Synced-CL 50919866

diff --git a/Engine/Extras/UnrealVS/VS2022/AtomFeed.xml b/Engine/Extras/UnrealVS/VS2022/AtomFeed.xml
--- a/Engine/Extras/UnrealVS/VS2022/AtomFeed.xml
+++ b/Engine/Extras/UnrealVS/VS2022/AtomFeed.xml
@@ -19,7 +19,7 @@
     <content type="application/octet-stream" src="UnrealVS.vsix" />
     <Vsix xmlns:xsd="http://www.w3.org/2001/XMLSchema" xmlns:xsi="http://www.w3.org/2001/XMLSchema-instance" xmlns="http://schemas.microsoft.com/developer/vsx-syndication-schema/2010">
       <Id>ddbf523f-7eb6-4887-bd51-85a714ff87eb</Id>
-      <Version>1.84.9.3</Version>
+      <Version>1.84.9.4</Version>
       <Rating xsi:nil="true" />
       <RatingCount xsi:nil="true" />
       <DownloadCount xsi:nil="true" />
diff --git a/Engine/Source/Programs/UnrealVS/ReleaseNotes.txt b/Engine/Source/Programs/UnrealVS/ReleaseNotes.txt
--- a/Engine/Source/Programs/UnrealVS/ReleaseNotes.txt
+++ b/Engine/Source/Programs/UnrealVS/ReleaseNotes.txt
@@ -1,3 +1,6 @@
+1.84.9.4
+* Adding UnrealVS.UBT options to build Modules from Solution Explorer folders.
+
 1.84.9.3
 * Addressing situations where "Compile and Profile Single File" didn't work correctly with Compile Score.
 * Disabling "Compile and Profile Single File" when Compile Score is not installed.
diff --git a/Engine/Source/Programs/UnrealVS/UnrealVS.2022/Properties/AssemblyInfo.cs b/Engine/Source/Programs/UnrealVS/UnrealVS.2022/Properties/AssemblyInfo.cs
--- a/Engine/Source/Programs/UnrealVS/UnrealVS.2022/Properties/AssemblyInfo.cs
+++ b/Engine/Source/Programs/UnrealVS/UnrealVS.2022/Properties/AssemblyInfo.cs
@@ -31,6 +31,6 @@ using System.Runtime.InteropServices;
 // You can specify all the values or you can default the Revision and Build Numbers 
 // by using the '*' as shown below:
 
-[assembly: AssemblyVersion("1.84.9.3")]
-[assembly: AssemblyFileVersion("1.84.9.3")]
+[assembly: AssemblyVersion("1.84.9.4")]
+[assembly: AssemblyFileVersion("1.84.9.4")]
 
diff --git a/Engine/Source/Programs/UnrealVS/UnrealVS.2022/UnrealVS.2022.csproj b/Engine/Source/Programs/UnrealVS/UnrealVS.2022/UnrealVS.2022.csproj
--- a/Engine/Source/Programs/UnrealVS/UnrealVS.2022/UnrealVS.2022.csproj
+++ b/Engine/Source/Programs/UnrealVS/UnrealVS.2022/UnrealVS.2022.csproj
@@ -223,7 +223,7 @@
   <Import Project="$(VSToolsPath)\VSSDK\Microsoft.VsSDK.targets" Condition="'$(VSToolsPath)' != ''" />
   <Import Project="$(MSBuildExtensionsPath32)\Microsoft\VisualStudio\v11.0\VSSDK\Microsoft.VsSDK.targets" Condition="false" />
   <PropertyGroup>
-    <PostBuildEvent>xcopy /F /Y "$(TargetDir)$(TargetName).vsix" "$(ProjectDir)..\..\..\..\Extras\UnrealVS\VS2022\$(TargetName).vsix"</PostBuildEvent>
+    <PostBuildEvent>if "$(ConfigurationName)" == "Release" xcopy /F /Y "$(TargetDir)$(TargetName).vsix" "$(ProjectDir)..\..\..\..\Extras\UnrealVS\VS2022\$(TargetName).vsix"</PostBuildEvent>
   </PropertyGroup>
   <PropertyGroup>
     <PreBuildEvent>
diff --git a/Engine/Source/Programs/UnrealVS/UnrealVS.Shared/BuildAndProfile.cs b/Engine/Source/Programs/UnrealVS/UnrealVS.Shared/BuildAndProfile.cs
--- a/Engine/Source/Programs/UnrealVS/UnrealVS.Shared/BuildAndProfile.cs
+++ b/Engine/Source/Programs/UnrealVS/UnrealVS.Shared/BuildAndProfile.cs
@@ -2,17 +2,46 @@
 
 using EnvDTE;
 using Microsoft.VisualStudio.Shell;
-using Microsoft.VisualStudio.Shell.Interop;
 using System;
 using System.Collections.Generic;
 using System.ComponentModel.Design;
-using System.Runtime.InteropServices;
 
 namespace UnrealVS
 {
 	class BuildAndProfile : IDisposable
 	{
-		private readonly int ProjectBuildAndProfileButtonID = 0x1071;
+		private static readonly int ProjectRebuildAndProfileButtonID = 0x1071;
+
+		private static readonly int FolderUnrealVSBuildButtonID = 0x1170;
+		private static readonly int FolderUnrealVSBuildNonUnityButtonID = 0x1171;
+		private static readonly int FolderUnrealVSRebuildAndProfileButtonID = 0x1172;
+		private static readonly int FolderUnrealVSRebuildAndProfileNonUnityButtonID = 0x1173;
+
+		private static bool IsProjectCommandID(int CommandID)
+		{
+			return CommandID == ProjectRebuildAndProfileButtonID;
+		}
+
+		private static bool IsProfileCommandID(int CommandID)
+		{
+			return CommandID == ProjectRebuildAndProfileButtonID
+				|| CommandID == FolderUnrealVSRebuildAndProfileButtonID
+				|| CommandID == FolderUnrealVSRebuildAndProfileNonUnityButtonID;
+		}
+
+		private static bool IsModuleBuildCommandID(int CommandID)
+		{
+			return CommandID == FolderUnrealVSBuildButtonID
+				|| CommandID == FolderUnrealVSBuildNonUnityButtonID
+				|| CommandID == FolderUnrealVSRebuildAndProfileButtonID
+				|| CommandID == FolderUnrealVSRebuildAndProfileNonUnityButtonID;
+		}
+
+		private static bool IsNonUnityBuildCommandID(int CommandID)
+		{
+			return CommandID == FolderUnrealVSBuildNonUnityButtonID
+				|| CommandID == FolderUnrealVSRebuildAndProfileNonUnityButtonID;
+		}
 
 		private UnrealProjectBuilder ProjectBuilder = new UnrealProjectBuilder();
 
@@ -20,9 +49,21 @@ namespace UnrealVS
 		{
 			ThreadHelper.ThrowIfNotOnUIThread();
 
-			var BuildAndProfileCommand = new OleMenuCommand(ProjectBuildAndProfileCallback, new CommandID(GuidList.UnrealVSCmdSet, ProjectBuildAndProfileButtonID));
-			BuildAndProfileCommand.BeforeQueryStatus += ProjectBuildAndProfileQueryStatus;
-			UnrealVSPackage.Instance.MenuCommandService.AddCommand(BuildAndProfileCommand);
+			// Project actions
+			AddCommandInternal(ProjectRebuildAndProfileButtonID);
+
+			// Folder actions
+			AddCommandInternal(FolderUnrealVSBuildButtonID);
+			AddCommandInternal(FolderUnrealVSBuildNonUnityButtonID);
+			AddCommandInternal(FolderUnrealVSRebuildAndProfileButtonID);
+			AddCommandInternal(FolderUnrealVSRebuildAndProfileNonUnityButtonID);
+
+			void AddCommandInternal(int CommandID)
+			{
+				var MenuCommand = new OleMenuCommand(BuildCallback, new CommandID(GuidList.UnrealVSCmdSet, CommandID));
+				MenuCommand.BeforeQueryStatus += BuildQueryStatus;
+				UnrealVSPackage.Instance.MenuCommandService.AddCommand(MenuCommand);
+			}
 		}
 
 		public void Dispose()
@@ -30,39 +71,87 @@ namespace UnrealVS
 			ProjectBuilder.Dispose();
 		}
 
-		private void ProjectBuildAndProfileQueryStatus(object Sender, EventArgs Args)
+		private void BuildQueryStatus(object Sender, EventArgs Args)
 		{
 			ThreadHelper.ThrowIfNotOnUIThread();
-			((MenuCommand)Sender).Enabled = CompileScoreHelper.IsInstalled();
-		}
 
-		private void ProjectBuildAndProfileCallback(object Sender, EventArgs Args)
-		{
-			ThreadHelper.ThrowIfNotOnUIThread();
+			if (Sender is MenuCommand BuildCommand)
+			{
+				bool bEnabled = true;
 
-			TryProjectBuildAndProfile();
+				if (IsProfileCommandID(BuildCommand.CommandID.ID) && !CompileScoreHelper.IsInstalled())
+				{
+					bEnabled = false;
+				}
+				else if (IsModuleBuildCommandID(BuildCommand.CommandID.ID) && !Utils.GetSelectedModule(out _, out _))
+				{
+					bEnabled = false;
+				}
+
+				BuildCommand.Enabled = bEnabled;
+			}
 		}
 
-		private bool TryProjectBuildAndProfile()
+		private void BuildCallback(object Sender, EventArgs Args)
 		{
 			ThreadHelper.ThrowIfNotOnUIThread();
 
-			// Get the project clicked in the solution explorer by accessing the current selection and converting to a Project if possible.
-			if (!(Utils.GetCurrentlySelectedProject() is Project SelectedProject))
+			if (Sender is MenuCommand BuildCommand)
 			{
-				return false;
-			}
+				List<string> ExtraArguments = new List<string>();
 
-			List<string> ExtraArguments = new List<string>
-			{
-				// When profiling:
-				//   Force disable UBA to force a local build.
-				//   Force use of MSVC since this is specifically for CompileScore.
-				//   Force a unique build environment, might also add more compile options.
-				"-NoUba -NoVfs -Compiler=VisualStudio -ProfilingCompile"
-			};
-
-			return ProjectBuilder.BuildProject(SelectedProject, null, ExtraArguments, true);
+				int CommandID = BuildCommand.CommandID.ID;
+
+				if (IsNonUnityBuildCommandID(CommandID))
+				{
+					ExtraArguments.Add("-DisableUnity");
+					// use -NoLink to prevent compiling all dependencies NonUnity as well.
+					ExtraArguments.Add("-NoLink");
+				}
+
+				if (IsProfileCommandID(CommandID))
+				{
+					ExtraArguments.AddRange(new List<string>
+					{
+						// When profiling:
+						//   Force disable UBA to force a local build.
+						"-NoUba",
+						"-NoVfs",
+						//   Force use of MSVC since this is specifically for CompileScore.
+						"-Compiler=VisualStudio",
+						//   Force a unique build environment, might also add more compile options.
+						"-ProfilingCompile",
+						//   Force a rebuild; by requesting a profile, you are asking for a full profile, not an incremental one.
+						"-Rebuild"
+					});
+
+					if (IsProjectCommandID(CommandID))
+					{
+						// Get the project clicked in the solution explorer by accessing the current selection and converting to a Project if possible.
+						if (Utils.GetCurrentlySelectedProject() is Project SelectedProject)
+						{
+							ProjectBuilder.BuildProject(SelectedProject, null, ExtraArguments, bProfiling: true);
+						}
+					}
+					else if (IsModuleBuildCommandID(CommandID))
+					{
+						// use -NoLink to prevent profiling all dependencies as well.
+						ExtraArguments.Add("-NoLink");
+
+						if (Utils.GetSelectedModule(out Project ModuleProject, out string ModuleName))
+						{
+							ProjectBuilder.BuildModule(ModuleProject, ModuleName, ExtraArguments, bProfiling: true);
+						}
+					}
+				}
+				else if (IsModuleBuildCommandID(CommandID))
+				{
+					if (Utils.GetSelectedModule(out Project ModuleProject, out string ModuleName))
+					{
+						ProjectBuilder.BuildModule(ModuleProject, ModuleName, ExtraArguments, bProfiling: false);
+					}
+				}
+			}
 		}
 	}
 }
diff --git a/Engine/Source/Programs/UnrealVS/UnrealVS.Shared/CompileSingleFile.cs b/Engine/Source/Programs/UnrealVS/UnrealVS.Shared/CompileSingleFile.cs
--- a/Engine/Source/Programs/UnrealVS/UnrealVS.Shared/CompileSingleFile.cs
+++ b/Engine/Source/Programs/UnrealVS/UnrealVS.Shared/CompileSingleFile.cs
@@ -159,7 +159,6 @@ namespace UnrealVS
 			string FileToCompile = DTE.ActiveDocument.FullName;
 			string FileToCompileExt = Path.GetExtension(FileToCompile);
 
-			string CompilingText;
 			List<string> ExtraArguments = new List<string>();
 
 			if (bPreProcessOnly)
@@ -181,8 +180,7 @@ namespace UnrealVS
 					return true;
 				}
 
-				CompilingText = FileToCompile;
-				ExtraArguments.Add($"-SingleFile=\"{FileToCompile}\"");
+				return ProjectBuilder.BuildSingleFile(StartupProject, FileToCompile, ExtraArguments, bProfile);
 			}
 			else
 			{
@@ -192,11 +190,8 @@ namespace UnrealVS
 					MessageBox.Show($"Can't find module for for {FileToCompile} to compile.", "Invalid Module", MessageBoxButtons.OK);
 					return true;
 				}
-				CompilingText = ModuleName;
-				ExtraArguments.Add($"-Module=\"{ModuleName}\"");
+				return ProjectBuilder.BuildModule(StartupProject, ModuleName, ExtraArguments, bProfile);
 			}
-
-			return ProjectBuilder.BuildProject(StartupProject, CompilingText, ExtraArguments, bProfile);
 		}
 	}
 }
diff --git a/Engine/Source/Programs/UnrealVS/UnrealVS.Shared/UnrealProjectBuilder.cs b/Engine/Source/Programs/UnrealVS/UnrealVS.Shared/UnrealProjectBuilder.cs
--- a/Engine/Source/Programs/UnrealVS/UnrealVS.Shared/UnrealProjectBuilder.cs
+++ b/Engine/Source/Programs/UnrealVS/UnrealVS.Shared/UnrealProjectBuilder.cs
@@ -8,6 +8,8 @@ using System;
 using System.Collections.Generic;
 using System.Diagnostics;
 using System.IO;
+using System.Linq;
+using System.Text.RegularExpressions;
 using System.Windows.Forms;
 
 namespace UnrealVS
@@ -26,7 +28,38 @@ namespace UnrealVS
 			KillChildProcess();
 		}
 
-		public bool BuildProject(Project TargetProject, string SpecificTarget, List<string> ExtraArguments, bool bRunningCompileScore)
+		/// <summary>
+		/// Removes any -Target="..." arguments after the first one from the provided string.
+		/// </summary>
+		/// <param name="BuildCommandLine">Set of arguments to reduce</param>
+		/// <returns>New, reduced command line</returns>
+		private static string RemoveExtraTargets(string BuildCommandLine)
+		{
+			if (string.IsNullOrEmpty(BuildCommandLine))
+			{
+				return BuildCommandLine;
+			}
+
+			// Match: -Target="..."
+			MatchCollection TargetMatches = Regex.Matches(BuildCommandLine, @"\s*-Target=""[^""]*""", RegexOptions.IgnoreCase);
+			if (TargetMatches.Count <= 1)
+			{
+				// No multi-target, return as-is
+				return BuildCommandLine;
+			}
+
+			string NewCommandLine = BuildCommandLine;
+
+			// Remove every -Target="" after the first one. Iterate in reverse order to preserve the match indices.
+			foreach (Match TargetMatch in TargetMatches.Cast<Match>().Skip(1).Reverse())
+			{
+				NewCommandLine = NewCommandLine.Remove(TargetMatch.Index, TargetMatch.Length);
+			}
+
+			return NewCommandLine;
+		}
+
+		public bool BuildProject(Project TargetProject, string SpecificTarget, List<string> ExtraArguments, bool bProfiling, bool bSingleTarget = false)
 		{
 			ThreadHelper.ThrowIfNotOnUIThread();
 			DTE DTE = UnrealVSPackage.Instance.DTE;
@@ -131,7 +164,7 @@ namespace UnrealVS
 				}
 			}
 
-			if (bRunningCompileScore && !CompileScoreHelper.StartTrace())
+			if (bProfiling && !CompileScoreHelper.StartTrace())
 			{
 				return false;
 			}
@@ -142,9 +175,15 @@ namespace UnrealVS
 				BuildOutputPane.OutputStringThreadSafe($"1>  Compiling {SpecificTarget}{Environment.NewLine}");
 			}
 
+			string NMakeBuildCommandLine = ActiveNMakeTool.BuildCommandLine;
+			if (bSingleTarget)
+			{
+				NMakeBuildCommandLine = RemoveExtraTargets(NMakeBuildCommandLine);
+			}
+
 			string SolutionDir = Path.GetDirectoryName(UnrealVSPackage.Instance.SolutionFilepath).TrimEnd(Path.DirectorySeparatorChar) + Path.DirectorySeparatorChar;
 			// Get the build command line and escape any environment variables that we use
-			string BuildCommandLine = ActiveVCConfiguration.Evaluate(ActiveNMakeTool.BuildCommandLine);
+			string BuildCommandLine = ActiveVCConfiguration.Evaluate(NMakeBuildCommandLine);
 			string UBTArgument = ActiveVCConfiguration.Evaluate(string.Join(" ", UBTArguments));
 			string WorkingDirectory = ActiveVCConfiguration.Evaluate("$(MSBuildProjectDirectory)");
 
@@ -163,7 +202,7 @@ namespace UnrealVS
 				// add an event handler to respond to the exit of the request
 				// and open the generated file if it exists.
 				ChildProcess.EnableRaisingEvents = true;
-				ChildProcess.Exited += new EventHandler((s, e) => UbtProcessExitHandler(PreprocessedFiles, AssemblyFiles, bRunningCompileScore));
+				ChildProcess.Exited += new EventHandler((s, e) => UbtProcessExitHandler(PreprocessedFiles, AssemblyFiles, bProfiling));
 			}
 
 			ChildProcess.Start();
@@ -173,6 +212,36 @@ namespace UnrealVS
 			return true;
 		}
 
+		public bool BuildSingleFile(Project TargetProject, string FileToCompile, List<string> InExtraArguments, bool bProfiling)
+		{
+			ThreadHelper.ThrowIfNotOnUIThread();
+
+			List<string> ExtraArguments = new List<string>()
+			{
+				$"-SingleFile=\"{FileToCompile}\""
+			};
+			if (InExtraArguments != null && InExtraArguments.Count > 0)
+			{
+				ExtraArguments.AddRange(InExtraArguments);
+			}
+			return BuildProject(TargetProject, FileToCompile, ExtraArguments, bProfiling: bProfiling);
+		}
+
+		public bool BuildModule(Project TargetProject, string ModuleName, List<string> InExtraArguments, bool bProfiling)
+		{
+			ThreadHelper.ThrowIfNotOnUIThread();
+
+			List<string> ExtraArguments = new List<string>()
+			{
+				$"-Module=\"{ModuleName}\""
+			};
+			if (InExtraArguments != null && InExtraArguments.Count > 0)
+			{
+				ExtraArguments.AddRange(InExtraArguments);
+			}
+			return BuildProject(TargetProject, ModuleName, ExtraArguments, bProfiling: bProfiling, bSingleTarget: true);
+		}
+
 		void KillChildProcess()
 		{
 			if (ChildProcess != null)
@@ -187,9 +256,9 @@ namespace UnrealVS
 			}
 		}
 
-		private void UbtProcessExitHandler(IEnumerable<string> PreprocessedFiles, IEnumerable<string> AssemblyFiles, bool bRunningCompileScore)
+		private void UbtProcessExitHandler(IEnumerable<string> PreprocessedFiles, IEnumerable<string> AssemblyFiles, bool bProfiling)
 		{
-			if (bRunningCompileScore)
+			if (bProfiling)
 			{
 				ThreadHelper.JoinableTaskFactory.Run(async () =>
 				{
diff --git a/Engine/Source/Programs/UnrealVS/UnrealVS.Shared/Utils.cs b/Engine/Source/Programs/UnrealVS/UnrealVS.Shared/Utils.cs
--- a/Engine/Source/Programs/UnrealVS/UnrealVS.Shared/Utils.cs
+++ b/Engine/Source/Programs/UnrealVS/UnrealVS.Shared/Utils.cs
@@ -13,7 +13,6 @@ using System.Linq;
 using System.Runtime.InteropServices;
 using System.Windows.Controls;
 using System.Windows.Media;
-using System.Windows.Media.Animation;
 using System.Xml;
 
 namespace UnrealVS
@@ -78,14 +77,14 @@ namespace UnrealVS
 
 			try
 			{
-				UnrealVSPackage.Instance.SelectionManager.GetCurrentSelection(out HierarchyPtr, out uint ProjectItemId, out IVsMultiItemSelect MultiItemSelect, out SelectionContainerPtr);
-
+				UnrealVSPackage.Instance.SelectionManager.GetCurrentSelection(out HierarchyPtr, out uint ProjectItemId, out _, out SelectionContainerPtr);
 				if (HierarchyPtr == IntPtr.Zero)
 				{
 					return null;
 				}
 
-				if (!(Marshal.GetTypedObjectForIUnknown(HierarchyPtr, typeof(IVsHierarchy)) is IVsHierarchy SelectedHierarchy))
+				object SelectedHierarchyObj = Marshal.GetObjectForIUnknown(HierarchyPtr);
+				if (!(SelectedHierarchyObj is IVsHierarchy SelectedHierarchy))
 				{
 					return null;
 				}
@@ -105,6 +104,93 @@ namespace UnrealVS
 			}
 		}
 
+		private static readonly string BuildModuleSuffix = ".Build.cs";
+
+		public static bool GetSelectedModule(out Project ModuleProject, out string ModuleName)
+		{
+			ThreadHelper.ThrowIfNotOnUIThread();
+
+			ModuleProject = null;
+			ModuleName = null;
+
+			IntPtr HierarchyPtr = IntPtr.Zero;
+			IntPtr SelectionContainerPtr = IntPtr.Zero;
+
+			try
+			{
+				UnrealVSPackage.Instance.SelectionManager.GetCurrentSelection(out HierarchyPtr, out uint ProjectItemId, out _, out SelectionContainerPtr);
+				if (HierarchyPtr == IntPtr.Zero)
+				{
+					return false;
+				}
+
+				object SelectedHierarchyObj = Marshal.GetObjectForIUnknown(HierarchyPtr);
+				if (!(SelectedHierarchyObj is IVsHierarchy SelectedHierarchy))
+				{
+					return false;
+				}
+
+				// Walk over the children and look for a .Build.cs file. If one exists, then this is a Module folder.
+				SelectedHierarchy.GetProperty(ProjectItemId, (int)__VSHPROPID.VSHPROPID_FirstChild, out object FirstChildIdObj);
+
+				uint GetItemID(object IdObject)
+				{
+					if (IdObject is null)
+					{
+						return VSConstants.VSITEMID_NIL;
+					}
+					if (IdObject is int IntId)
+					{
+						return (uint)IntId;
+					}
+					if (IdObject is uint UintId)
+					{
+						return UintId;
+					}
+					return VSConstants.VSITEMID_NIL;
+				}
+
+				if (FirstChildIdObj != null)
+				{
+					uint ChildObjectID = GetItemID(FirstChildIdObj);
+
+					while (ChildObjectID != VSConstants.VSITEMID_NIL)
+					{
+						SelectedHierarchy.GetProperty(ChildObjectID, (int)__VSHPROPID.VSHPROPID_Name, out object ObjectNameObj);
+
+						if (ObjectNameObj is string ObjectName && ObjectName.EndsWith(BuildModuleSuffix, StringComparison.InvariantCultureIgnoreCase))
+						{
+							SelectedHierarchy.GetProperty(ProjectItemId, (int)__VSHPROPID.VSHPROPID_ExtObject, out object ProjectObject);
+
+							if (ProjectObject is ProjectItem SelectedItem)
+							{
+								ModuleName = ObjectName.Substring(0, ObjectName.Length - BuildModuleSuffix.Length);
+								ModuleProject = SelectedItem.ContainingProject;
+								return true;
+							}
+						}
+
+						SelectedHierarchy.GetProperty(ChildObjectID, (int)__VSHPROPID.VSHPROPID_NextSibling, out object SiblingIdObj);
+
+						ChildObjectID = GetItemID(SiblingIdObj);
+					}
+				}
+
+				return false;
+			}
+			finally
+			{
+				if (HierarchyPtr != IntPtr.Zero)
+				{
+					Marshal.Release(HierarchyPtr);
+				}
+				if (SelectionContainerPtr != IntPtr.Zero)
+				{
+					Marshal.Release(SelectionContainerPtr);
+				}
+			}
+		}
+
 		/// <summary>
 		/// Converts an IVsHierarchy object to a config provider interface
 		/// </summary>
