[UnrealVS]
- Version 1.84.9.5
- Adding "Copy Unreal #include Statement" context menu item to .h document tabs in Visual Studio. Result is copied as: #include "Path/To/Header.h"
- Made smaller version of the Unreal icon to use as a context menu icon.

#Synced-CL 51790204

diff --git a/Engine/Extras/UnrealVS/VS2022/AtomFeed.xml b/Engine/Extras/UnrealVS/VS2022/AtomFeed.xml
--- a/Engine/Extras/UnrealVS/VS2022/AtomFeed.xml
+++ b/Engine/Extras/UnrealVS/VS2022/AtomFeed.xml
@@ -19,7 +19,7 @@
     <content type="application/octet-stream" src="UnrealVS.vsix" />
     <Vsix xmlns:xsd="http://www.w3.org/2001/XMLSchema" xmlns:xsi="http://www.w3.org/2001/XMLSchema-instance" xmlns="http://schemas.microsoft.com/developer/vsx-syndication-schema/2010">
       <Id>ddbf523f-7eb6-4887-bd51-85a714ff87eb</Id>
-      <Version>1.84.9.4</Version>
+      <Version>1.84.9.5</Version>
       <Rating xsi:nil="true" />
       <RatingCount xsi:nil="true" />
       <DownloadCount xsi:nil="true" />
diff --git a/Engine/Source/Programs/UnrealVS/ReleaseNotes.txt b/Engine/Source/Programs/UnrealVS/ReleaseNotes.txt
--- a/Engine/Source/Programs/UnrealVS/ReleaseNotes.txt
+++ b/Engine/Source/Programs/UnrealVS/ReleaseNotes.txt
@@ -1,3 +1,6 @@
+1.84.9.5
+* Adding "Copy Unreal #include Statement" to the document tab context menu.
+
 1.84.9.4
 * Adding UnrealVS.UBT options to build Modules from Solution Explorer folders.
 
diff --git a/Engine/Source/Programs/UnrealVS/UnrealVS.2022/Properties/AssemblyInfo.cs b/Engine/Source/Programs/UnrealVS/UnrealVS.2022/Properties/AssemblyInfo.cs
--- a/Engine/Source/Programs/UnrealVS/UnrealVS.2022/Properties/AssemblyInfo.cs
+++ b/Engine/Source/Programs/UnrealVS/UnrealVS.2022/Properties/AssemblyInfo.cs
@@ -31,6 +31,6 @@ using System.Runtime.InteropServices;
 // You can specify all the values or you can default the Revision and Build Numbers 
 // by using the '*' as shown below:
 
-[assembly: AssemblyVersion("1.84.9.4")]
-[assembly: AssemblyFileVersion("1.84.9.4")]
+[assembly: AssemblyVersion("1.84.9.5")]
+[assembly: AssemblyFileVersion("1.84.9.5")]
 
diff --git a/Engine/Source/Programs/UnrealVS/UnrealVS.2022/UnrealVS.2022.csproj b/Engine/Source/Programs/UnrealVS/UnrealVS.2022/UnrealVS.2022.csproj
--- a/Engine/Source/Programs/UnrealVS/UnrealVS.2022/UnrealVS.2022.csproj
+++ b/Engine/Source/Programs/UnrealVS/UnrealVS.2022/UnrealVS.2022.csproj
@@ -112,6 +112,8 @@
     </Resource>
     <Resource Include="Resources\RefreshProjects.png">
     </Resource>
+    <Resource Include="Resources\UnrealIcon.png">
+    </Resource>
     <None Include="source.extension.vsixmanifest" />
     <Content Include="..\ReleaseNotes.txt">
       <Visible>false</Visible>
diff --git a/Engine/Source/Programs/UnrealVS/UnrealVS.Shared/CopyIncludePath.cs b/Engine/Source/Programs/UnrealVS/UnrealVS.Shared/CopyIncludePath.cs
new file mode 100644
--- /dev/null
+++ b/Engine/Source/Programs/UnrealVS/UnrealVS.Shared/CopyIncludePath.cs
@@ -0,0 +1,138 @@
+// Copyright Epic Games, Inc. All Rights Reserved.
+
+using EnvDTE;
+using Microsoft.VisualStudio.Shell;
+using System;
+using System.Collections.Generic;
+using System.ComponentModel.Design;
+using System.IO;
+using System.Windows.Forms;
+
+namespace UnrealVS
+{
+	class CopyIncludePath : IDisposable
+	{
+		private const int CopyIncludePathButtonID = 0x1470;
+
+		/// <summary>
+		/// Well-known module subdirectory names that serve as include roots in UE.
+		/// When a header lives under one of these directories the include path
+		/// starts immediately after it (e.g. "MyModule/Public/Foo/Bar.h" becomes
+		/// "Foo/Bar.h").
+		/// </summary>
+		private static readonly HashSet<string> IncludeRootFolders =
+			new HashSet<string>(StringComparer.OrdinalIgnoreCase)
+			{
+				"Public",
+				"Private",
+				"Classes",
+				"Internal",
+			};
+
+		private readonly OleMenuCommand MenuCommand;
+
+		public CopyIncludePath()
+		{
+			MenuCommand = new OleMenuCommand(OnExecute, new CommandID(GuidList.UnrealVSCmdSet, CopyIncludePathButtonID));
+			MenuCommand.BeforeQueryStatus += OnBeforeQueryStatus;
+			UnrealVSPackage.Instance.MenuCommandService.AddCommand(MenuCommand);
+		}
+
+		public void Dispose()
+		{
+		}
+
+		/// <summary>
+		/// Only show the menu item when the active document is a valid Unreal header
+		/// (a .h file inside a UE solution that lives under a known module include root).
+		/// </summary>
+		private void OnBeforeQueryStatus(object sender, EventArgs e)
+		{
+			ThreadHelper.ThrowIfNotOnUIThread();
+
+			bool bVisible = false;
+			DTE DTE = UnrealVSPackage.Instance.DTE;
+
+			if (UnrealVSPackage.Instance.IsUESolutionLoaded
+				&& DTE?.ActiveDocument != null
+				&& Path.GetExtension(DTE.ActiveDocument.FullName).Equals(".h", StringComparison.OrdinalIgnoreCase))
+			{
+				bVisible = BuildIncludePath(DTE.ActiveDocument.FullName, out _);
+			}
+
+			MenuCommand.Visible = bVisible;
+			MenuCommand.Enabled = bVisible;
+		}
+
+		/// <summary>
+		/// Build a #include "..." string from the active document path and copy it
+		/// to the clipboard.
+		/// </summary>
+		private void OnExecute(object sender, EventArgs e)
+		{
+			ThreadHelper.ThrowIfNotOnUIThread();
+
+			DTE DTE = UnrealVSPackage.Instance.DTE;
+			if (DTE?.ActiveDocument == null)
+			{
+				return;
+			}
+
+			string FullPath = DTE.ActiveDocument.FullName;
+			if (!BuildIncludePath(FullPath, out string IncludePath))
+			{
+				return;
+			}
+
+			string IncludeDirective = $"#include \"{IncludePath}\"";
+
+			Clipboard.SetText(IncludeDirective);
+
+			// Brief status bar message so the user knows it worked
+			DTE.StatusBar.Text = $"Copied: {IncludeDirective}";
+		}
+
+		/// <summary>
+		/// Determines the shortest conventional UE #include path for a header.
+		///
+		/// Strategy: walk up from the file looking for a well-known include-root
+		/// folder (Public, Private, Classes, Internal). The include path is
+		/// everything after that folder.
+		///
+		/// Returns false if no include root folder is found in the path.
+		/// </summary>
+		internal static bool BuildIncludePath(string FullPath, out string IncludePath)
+		{
+			FullPath = Path.GetFullPath(FullPath);
+
+			// Collect path segments from the file upward
+			List<string> Segments = new List<string>();
+			string Dir = Path.GetDirectoryName(FullPath);
+			Segments.Add(Path.GetFileName(FullPath));
+
+			while (!string.IsNullOrEmpty(Dir))
+			{
+				string FolderName = Path.GetFileName(Dir);
+
+				if (IncludeRootFolders.Contains(FolderName))
+				{
+					// Everything we've accumulated so far (in reverse) is the include path
+					Segments.Reverse();
+					IncludePath = string.Join("/", Segments);
+					return true;
+				}
+
+				Segments.Add(FolderName);
+				string Parent = Path.GetDirectoryName(Dir);
+				if (Parent == Dir)
+				{
+					break; // root
+				}
+				Dir = Parent;
+			}
+
+			IncludePath = null;
+			return false;
+		}
+	}
+}
diff --git a/Engine/Source/Programs/UnrealVS/UnrealVS.Shared/UnrealVSPackage.cs b/Engine/Source/Programs/UnrealVS/UnrealVS.Shared/UnrealVSPackage.cs
--- a/Engine/Source/Programs/UnrealVS/UnrealVS.Shared/UnrealVSPackage.cs
+++ b/Engine/Source/Programs/UnrealVS/UnrealVS.Shared/UnrealVSPackage.cs
@@ -295,6 +295,9 @@ namespace UnrealVS
 			// Create 'UbaVisualizer' instance
 			UbaVisualizer = new UbaVisualizer();
 
+			// Create 'CopyIncludePath' instance
+			CopyIncludePath = new CopyIncludePath();
+
 			// Call parent implementation
 			base.Initialize();
 
@@ -456,6 +459,9 @@ namespace UnrealVS
 			P4CommandsGroup?.Dispose();
 			P4CommandsGroup = null;
 
+			CopyIncludePath?.Dispose();
+			CopyIncludePath = null;
+
 			CommandLineEditor?.Dispose();
 			CommandLineEditor = null;
 
@@ -1122,6 +1128,9 @@ namespace UnrealVS
 		/// Batch Builder button/command handler
 		private BatchBuilder BatchBuilder = new BatchBuilder();
 
+		/// CopyIncludePath feature
+		private CopyIncludePath CopyIncludePath;
+
 		/// Perforce features
 		private P4Commands P4CommandsGroup;
 
