UnrealVS: Usability improvements for startup project selector

#rnx
#Synced-CL 52111506

diff --git a/Engine/Extras/UnrealVS/VS2022/AtomFeed.xml b/Engine/Extras/UnrealVS/VS2022/AtomFeed.xml
--- a/Engine/Extras/UnrealVS/VS2022/AtomFeed.xml
+++ b/Engine/Extras/UnrealVS/VS2022/AtomFeed.xml
@@ -19,7 +19,7 @@
     <content type="application/octet-stream" src="UnrealVS.vsix" />
     <Vsix xmlns:xsd="http://www.w3.org/2001/XMLSchema" xmlns:xsi="http://www.w3.org/2001/XMLSchema-instance" xmlns="http://schemas.microsoft.com/developer/vsx-syndication-schema/2010">
       <Id>ddbf523f-7eb6-4887-bd51-85a714ff87eb</Id>
-      <Version>1.84.9.5</Version>
+      <Version>1.84.9.6</Version>
       <Rating xsi:nil="true" />
       <RatingCount xsi:nil="true" />
       <DownloadCount xsi:nil="true" />
diff --git a/Engine/Source/Programs/UnrealVS/ReleaseNotes.txt b/Engine/Source/Programs/UnrealVS/ReleaseNotes.txt
--- a/Engine/Source/Programs/UnrealVS/ReleaseNotes.txt
+++ b/Engine/Source/Programs/UnrealVS/ReleaseNotes.txt
@@ -1,3 +1,6 @@
+1.84.9.6
+* Usability improvements for startup project selector
+
 1.84.9.5
 * Adding "Copy Unreal #include Statement" to the document tab context menu.
 
diff --git a/Engine/Source/Programs/UnrealVS/UnrealVS.2022/Properties/AssemblyInfo.cs b/Engine/Source/Programs/UnrealVS/UnrealVS.2022/Properties/AssemblyInfo.cs
--- a/Engine/Source/Programs/UnrealVS/UnrealVS.2022/Properties/AssemblyInfo.cs
+++ b/Engine/Source/Programs/UnrealVS/UnrealVS.2022/Properties/AssemblyInfo.cs
@@ -31,6 +31,6 @@ using System.Runtime.InteropServices;
 // You can specify all the values or you can default the Revision and Build Numbers 
 // by using the '*' as shown below:
 
-[assembly: AssemblyVersion("1.84.9.5")]
-[assembly: AssemblyFileVersion("1.84.9.5")]
+[assembly: AssemblyVersion("1.84.9.6")]
+[assembly: AssemblyFileVersion("1.84.9.6")]
 
diff --git a/Engine/Source/Programs/UnrealVS/UnrealVS.Shared/CommandLineEditor.cs b/Engine/Source/Programs/UnrealVS/UnrealVS.Shared/CommandLineEditor.cs
--- a/Engine/Source/Programs/UnrealVS/UnrealVS.Shared/CommandLineEditor.cs
+++ b/Engine/Source/Programs/UnrealVS/UnrealVS.Shared/CommandLineEditor.cs
@@ -192,8 +192,9 @@ namespace UnrealVS
 				ComboCommandList.Enabled = false;
 			}
 
-			string CommandLine = MakeCommandLineComboText();
-			CommitCommandLineToMRU(CommandLine);
+			// Mark dirty: MakeCommandLineComboText() will run lazily when VS next polls OutValue.
+			// Avoids 5+ COM calls on every OnStartupProjectChanged event (e.g. hover over combo items).
+			_CachedCommandLineTextDirty = true;
 		}
 
 
@@ -318,7 +319,17 @@ namespace UnrealVS
 						}
 						else
 						{
-							TextToDisplay = MakeCommandLineComboText();
+							if (_CachedCommandLineTextDirty)
+							{
+								string NewText = MakeCommandLineComboText();
+								if (NewText != CachedCommandLineText)
+								{
+									CachedCommandLineText = NewText;
+									CommitCommandLineToMRU(CachedCommandLineText);
+								}
+								_CachedCommandLineTextDirty = false;
+							}
+							TextToDisplay = CachedCommandLineText;
 						}
 					}
 
@@ -530,6 +541,11 @@ namespace UnrealVS
 		/// Used to store the user edited commandline mid-edit, in the combo handler
 		private string DesiredCommandLine;
 
+		/// Cached result of MakeCommandLineComboText(). Lazily computed: set dirty on state change,
+		/// recomputed only when VS polls OutValue. Avoids expensive COM calls on every project change event.
+		private string CachedCommandLineText = string.Empty;
+		private bool _CachedCommandLineTextDirty = true;
+
 		/// used to shield against too many property changes in the future
 		private System.Timers.Timer UpdateCommandLineComboTimer = null;
 
diff --git a/Engine/Source/Programs/UnrealVS/UnrealVS.Shared/StartupProjectSelector.cs b/Engine/Source/Programs/UnrealVS/UnrealVS.Shared/StartupProjectSelector.cs
--- a/Engine/Source/Programs/UnrealVS/UnrealVS.Shared/StartupProjectSelector.cs
+++ b/Engine/Source/Programs/UnrealVS/UnrealVS.Shared/StartupProjectSelector.cs
@@ -80,6 +80,7 @@ namespace UnrealVS
 					_bIsSolutionOpened = false;
 					UpdateStartupProjectCombo();
 					_CachedStartupProjects.Clear();
+					_CachedStartupProjectNames = null;
 				};
 			UnrealVSPackage.Instance.OnProjectOpened +=
 				delegate (Project OpenedProject)
@@ -152,6 +153,7 @@ namespace UnrealVS
 
 			SortCachedStartupProjects();
 
+			_CachedStartupProjectNames = null; // Invalidate: list changed
 			StartupProjectListChanged?.Invoke(this, new StartupProjectListChangedEventArgs { StartupProjects = StartupProjectList });
 		}
 
@@ -200,6 +202,7 @@ namespace UnrealVS
 		private void RemoveFromStartupProjectList(Project Project)
 		{
 			_CachedStartupProjects.RemoveAll(ProjRef => ProjRef.Project == Project);
+			_CachedStartupProjectNames = null; // Invalidate: list changed
 
 			StartupProjectListChanged?.Invoke(this, new StartupProjectListChangedEventArgs { StartupProjects = StartupProjectList });
 		}
@@ -209,10 +212,12 @@ namespace UnrealVS
 		/// </summary>
 		private void UpdateStartupProjectCombo()
 		{
+			ThreadHelper.ThrowIfNotOnUIThread();
 			foreach (var Command in StartupProjectComboCommands)
 			{
 				Command.Enabled = _bIsSolutionOpened;
 			}
+			_CachedStartupProjectComboText = _bIsSolutionOpened ? MakeStartupProjectComboText() : string.Empty;
 		}
 
 
@@ -221,6 +226,7 @@ namespace UnrealVS
 		/// </summary>
 		public void OnStartupProjectChanged(Project NewStartupProject)
 		{
+			ThreadHelper.ThrowIfNotOnUIThread();
 			UpdateStartupProjectCombo();
 		}
 
@@ -260,8 +266,7 @@ namespace UnrealVS
 			// If OutValue is non-zero, then Visual Studio is querying the current combo value to display
 			if (OleArgs.OutValue != IntPtr.Zero)
 			{
-				string ValueString = MakeStartupProjectComboText();
-				Marshal.GetNativeVariantForObject(ValueString, OleArgs.OutValue);
+				Marshal.GetNativeVariantForObject(_CachedStartupProjectComboText, OleArgs.OutValue);
 			}
 
 			// Has a new input value been entered?
@@ -269,21 +274,24 @@ namespace UnrealVS
 			{
 				var InputString = OleArgs.InValue.ToString();
 
-				// Lookup the ProjectReference in CachedStartupProjects to get the full name
-				ProjectReference ProjectRefMatch = _CachedStartupProjects.FirstOrDefault(ProjRef => ProjRef.Name == InputString);
+				if (string.IsNullOrWhiteSpace(InputString))
+				{
+					return;
+				}
+
+				// Lookup the ProjectReference in CachedStartupProjects to get the full name.
+				// Support partial/substring input: try exact match first, then prefix, then contains.
+				ProjectReference ProjectRefMatch =
+					_CachedStartupProjects.FirstOrDefault(ProjRef => string.Equals(ProjRef.Name, InputString, StringComparison.OrdinalIgnoreCase))
+					?? _CachedStartupProjects.FirstOrDefault(ProjRef => ProjRef.Name.StartsWith(InputString, StringComparison.OrdinalIgnoreCase))
+					?? _CachedStartupProjects.FirstOrDefault(ProjRef => ProjRef.Name.IndexOf(InputString, StringComparison.OrdinalIgnoreCase) >= 0);
 
 				if (ProjectRefMatch != null && ProjectRefMatch.Project != null)
 				{
-					// Switch to this project using the standard SDK method set_StartupProject
+					// Switch to this project using the standard SDK method set_StartupProject.
+					// set_StartupProject fires the proper VS events and marks the solution dirty for .suo persistence.
 					var ProjectHierarchy = Utils.ProjectToHierarchyObject(ProjectRefMatch.Project);
 					UnrealVSPackage.Instance.SolutionBuildManager.set_StartupProject(ProjectHierarchy);
-
-					// Select in solution exp and set as startup project
-					// This forces the .suo to save the project change
-					if (Utils.SelectProjectInSolutionExplorer(ProjectRefMatch.Project))
-					{
-						UnrealVSPackage.Instance.DTE.ExecuteCommand("Project.SetasStartUpProject");
-					}
 				}
 			}
 		}
@@ -295,8 +303,12 @@ namespace UnrealVS
 
 			var OleArgs = (OleMenuCmdEventArgs)Args;
 
-			var CachedStartupProjectNames = (from ProjRef in _CachedStartupProjects select ProjRef.Name).ToArray();
-			Marshal.GetNativeVariantForObject(CachedStartupProjectNames, OleArgs.OutValue);
+			// Rebuild names array only when the project list has changed, not on every dropdown open
+			if (_CachedStartupProjectNames == null)
+			{
+				_CachedStartupProjectNames = (from ProjRef in _CachedStartupProjects select ProjRef.Name).ToArray();
+			}
+			Marshal.GetNativeVariantForObject(_CachedStartupProjectNames, OleArgs.OutValue);
 		}
 
 		/// Combo control and list commands for our startup project selector
@@ -307,5 +319,11 @@ namespace UnrealVS
 
 		/// Store the list of projects that should be shown in the dropdown list
 		List<ProjectReference> _CachedStartupProjects = new List<ProjectReference>();
+
+		/// Cached names array for the combo list handler, null means dirty/needs rebuild
+		string[] _CachedStartupProjectNames = null;
+
+		/// Cached display text for the startup project combo, avoids repeated COM calls on every VS poll
+		string _CachedStartupProjectComboText = string.Empty;
 	}
 }
diff --git a/Engine/Source/Programs/UnrealVS/UnrealVS.Shared/Utils.cs b/Engine/Source/Programs/UnrealVS/UnrealVS.Shared/Utils.cs
--- a/Engine/Source/Programs/UnrealVS/UnrealVS.Shared/Utils.cs
+++ b/Engine/Source/Programs/UnrealVS/UnrealVS.Shared/Utils.cs
@@ -797,28 +797,38 @@ namespace UnrealVS
 		{
 			ThreadHelper.ThrowIfNotOnUIThread();
 
-			UnrealVSPackage.Instance.DTE.ExecuteCommand("View.SolutionExplorer");
-			if (Project.ParentProjectItem != null)
+			// Build path to the project in solution explorer: "SolutionName\\FolderName\\...\\ProjectName"
+			// Use UIHierarchy.GetItem() for direct O(depth) lookup instead of walking the entire tree.
+			var PathParts = new List<string> { Project.Name };
+			var ParentItem = Project.ParentProjectItem;
+			while (ParentItem != null)
 			{
-				Project.ParentProjectItem.ExpandView();
+				var ParentProject = ParentItem.ContainingProject;
+				if (ParentProject == null) break;
+				PathParts.Insert(0, ParentProject.Name);
+				ParentItem = ParentProject.ParentProjectItem;
 			}
+			string SolutionName = Path.GetFileNameWithoutExtension(UnrealVSPackage.Instance.DTE.Solution.FullName);
+			if (string.IsNullOrEmpty(SolutionName))
+			{
+				return false;
+			}
+			PathParts.Insert(0, SolutionName);
+			string ItemPath = string.Join("\\", PathParts);
 
-			UIHierarchy SolutionExplorerHierarachy = UnrealVSPackage.Instance.DTE2.ToolWindows.SolutionExplorer;
-			Utils.UITreeItem SolutionExplorerTree = Utils.GetUIHierarchyTree(SolutionExplorerHierarachy);
-			var UIHierarachyProjects = Utils.GetUITreeItemsByObjectType<Project>(SolutionExplorerTree);
-
-#pragma warning disable VSTHRD010 // Invoke single-threaded types on Main thread
-			var SelectableUIItem = UIHierarachyProjects.FirstOrDefault(uihp => uihp.Object as Project == Project);
-#pragma warning restore VSTHRD010 // Invoke single-threaded types on Main thread
-
-			if (SelectableUIItem != null)
+			UIHierarchy SolutionExplorer = UnrealVSPackage.Instance.DTE2.ToolWindows.SolutionExplorer;
+			try
 			{
-				if (Project.ParentProjectItem != null)
+				UIHierarchyItem Item = SolutionExplorer.GetItem(ItemPath);
+				if (Item != null)
 				{
-					SelectableUIItem.Select(vsUISelectionType.vsUISelectionTypeSelect);
+					Item.Select(vsUISelectionType.vsUISelectionTypeSelect);
 					return true;
 				}
 			}
+			catch (Exception)
+			{
+			}
 
 			return false;
 		}
