[UnrealVS] Apply Epic coding standard to FileBrowserWindowControl.

#rnx
#Synced-CL 52302118
#Synced-CL 52302118

diff --git a/Engine/Source/Programs/UnrealVS/UnrealVS.Shared/FileBrowserWindowControl.xaml.cs b/Engine/Source/Programs/UnrealVS/UnrealVS.Shared/FileBrowserWindowControl.xaml.cs
--- a/Engine/Source/Programs/UnrealVS/UnrealVS.Shared/FileBrowserWindowControl.xaml.cs
+++ b/Engine/Source/Programs/UnrealVS/UnrealVS.Shared/FileBrowserWindowControl.xaml.cs
@@ -21,10 +21,10 @@ namespace UnrealVS
 	[global::System.Diagnostics.CodeAnalysis.SuppressMessageAttribute("", "VSTHRD010")]
 	public partial class FileBrowserWindowControl : UserControl
 	{
-		FileBrowserWindow Window;
-		bool ClearFilter;
-		bool IsRefreshing;
-		ListBox ActiveFilesListBox;
+		private FileBrowserWindow Window;
+		private bool bClearFilter;
+		private bool bIsRefreshing;
+		private ListBox ActiveFilesListBox;
 
 		public class FileItem
 		{
@@ -34,32 +34,32 @@ namespace UnrealVS
 			public bool Selectable { get; set; }
 		}
 
-		String UnrealVsFileName;
+		private string UnrealVsFileName;
 
-		FileItem[] AllFileItems;
-		List<FileItem> BookmarkedFileItems;
-		List<FileItem> RecentFileItems;
+		private FileItem[] AllFileItems;
+		private List<FileItem> BookmarkedFileItems;
+		private List<FileItem> RecentFileItems;
 
-		Dictionary<string, FileItem> AllFileItemsLookup;
+		private Dictionary<string, FileItem> AllFileItemsLookup;
 
-		public FileBrowserWindowControl(FileBrowserWindow window)
+		public FileBrowserWindowControl(FileBrowserWindow InWindow)
 		{
-			Window = window;
-			this.InitializeComponent();
+			Window = InWindow;
+			InitializeComponent();
 
 			AllFileItems = Array.Empty<FileItem>();
 			BookmarkedFileItems = new List<FileItem>();
 			RecentFileItems = new List<FileItem>();
 			IsVisibleChanged += FileBrowserWindowControl_IsVisibleChanged;
-			foreach (var listView in new[] { AllFilesListBox, BookmarkedFilesListBox, RecentFilesListBox })
+			foreach (ListBox FileListBox in new[] { AllFilesListBox, BookmarkedFilesListBox, RecentFilesListBox })
 			{
-				listView.ItemsSource = new List<FileItem>();
-				listView.KeyDown += FileListView_KeyDown;
-				listView.PreviewKeyDown += FileListView_PreviewKeyDown;
-				listView.SelectionChanged += FileListView_SelectionChanged;
-				listView.GotFocus += FileListView_GotFocus;
-				listView.MouseDoubleClick += ListView_MouseDoubleClick;
-				listView.SelectionMode = SelectionMode.Extended;
+				FileListBox.ItemsSource = new List<FileItem>();
+				FileListBox.KeyDown += FileListView_KeyDown;
+				FileListBox.PreviewKeyDown += FileListView_PreviewKeyDown;
+				FileListBox.SelectionChanged += FileListView_SelectionChanged;
+				FileListBox.GotFocus += FileListView_GotFocus;
+				FileListBox.MouseDoubleClick += ListView_MouseDoubleClick;
+				FileListBox.SelectionMode = SelectionMode.Extended;
 			}
 			FilterEditBox.TextChanged += FilterEditBox_TextChanged;
 			FilterEditBox.KeyDown += FilterEditBox_KeyDown;
@@ -69,60 +69,74 @@ namespace UnrealVS
 			FilesListTab.SelectedIndex = 0;
 
 			ActiveFilesListBox = AllFilesListBox;
-			this.PreviewKeyDown += FileBrowserWindowControl_PreviewKeyDown;
+			PreviewKeyDown += FileBrowserWindowControl_PreviewKeyDown;
 		}
 
-		private void FilterEditBox_OnPaste(object sender, DataObjectPastingEventArgs e)
+		private void FilterEditBox_OnPaste(object Sender, DataObjectPastingEventArgs E)
 		{
-			var isText = e.SourceDataObject.GetDataPresent(DataFormats.UnicodeText, true);
-			if (!isText)
+			bool bIsText = E.SourceDataObject.GetDataPresent(DataFormats.UnicodeText, true);
+			if (!bIsText)
+			{
 				return;
+			}
 
-			var text = e.SourceDataObject.GetData(DataFormats.UnicodeText) as string;
+			string Text = E.SourceDataObject.GetData(DataFormats.UnicodeText) as string;
 
-			if (string.IsNullOrEmpty(text) || text.IndexOfAny(Path.GetInvalidPathChars()) != -1)
+			if (string.IsNullOrEmpty(Text) || Text.IndexOfAny(Path.GetInvalidPathChars()) != -1)
+			{
 				return;
+			}
 
-			FilterEditBox.Text = Path.GetFileName(text);
-			e.CancelCommand();
+			FilterEditBox.Text = Path.GetFileName(Text);
+			E.CancelCommand();
 			FilterEditBox.CaretIndex = FilterEditBox.Text.Length;
 		}
 
 		[STAThread]
-		private void FileBrowserWindowControl_PreviewKeyDown(object sender, KeyEventArgs e)
+		private void FileBrowserWindowControl_PreviewKeyDown(object Sender, KeyEventArgs E)
 		{
 			if (FilterEditBox.IsKeyboardFocused)
+			{
 				return;
+			}
 
-			if (Keyboard.Modifiers.HasFlag(ModifierKeys.Control) && e.Key == Key.V || Keyboard.Modifiers.HasFlag(ModifierKeys.Shift) && e.Key == Key.Insert)
+			if (Keyboard.Modifiers.HasFlag(ModifierKeys.Control) && E.Key == Key.V || Keyboard.Modifiers.HasFlag(ModifierKeys.Shift) && E.Key == Key.Insert)
 			{
-				string text = Clipboard.GetText();
+				string Text = Clipboard.GetText();
 
-				if (string.IsNullOrEmpty(text) || text.IndexOfAny(Path.GetInvalidPathChars()) != -1)
+				if (string.IsNullOrEmpty(Text) || Text.IndexOfAny(Path.GetInvalidPathChars()) != -1)
+				{
 					return;
+				}
 
-				FilterEditBox.Text = Path.GetFileName(text);
+				FilterEditBox.Text = Path.GetFileName(Text);
 				FilterEditBox.SelectAll();
 			}
 		}
 
-		private void FilesListTab_SelectionChanged(object sender, SelectionChangedEventArgs e)
+		private void FilesListTab_SelectionChanged(object Sender, SelectionChangedEventArgs E)
 		{
 			switch (FilesListTab.SelectedIndex)
 			{
 				case 0:
 					if (ActiveFilesListBox == AllFilesListBox)
+					{
 						return;
+					}
 					ActiveFilesListBox = AllFilesListBox;
 					break;
 				case 1:
 					if (ActiveFilesListBox == BookmarkedFilesListBox)
+					{
 						return;
+					}
 					ActiveFilesListBox = BookmarkedFilesListBox;
 					break;
 				case 2:
 					if (ActiveFilesListBox == RecentFilesListBox)
+					{
 						return;
+					}
 					ActiveFilesListBox = RecentFilesListBox;
 					break;
 			}
@@ -131,121 +145,134 @@ namespace UnrealVS
 			AsyncFocusSelectedItem();
 		}
 
-		private void FilterEditBox_KeyDown(object sender, KeyEventArgs e)
+		private void FilterEditBox_KeyDown(object Sender, KeyEventArgs E)
 		{
-			if (e.Key == Key.Tab)
+			if (E.Key == Key.Tab)
 			{
-				e.Handled = true;
+				E.Handled = true;
 				FocusListViewItem(ActiveFilesListBox, ActiveFilesListBox.SelectedIndex);
 				return;
 			}
 		}
-		private void FilterEditBox_PreviewKeyDown(object sender, KeyEventArgs e)
+
+		private void FilterEditBox_PreviewKeyDown(object Sender, KeyEventArgs E)
 		{
-			var listBox = ActiveFilesListBox;
-			var getSelectedItem = new Func<FileItem>(() =>
+			ListBox TargetListBox = ActiveFilesListBox;
+			Func<FileItem> GetSelectedItem = new Func<FileItem>(() =>
 			{
-				if (listBox.SelectedItems.Count == 0)
-					return (FileItem)listBox.Items[0];
+				if (TargetListBox.SelectedItems.Count == 0)
+					return (FileItem)TargetListBox.Items[0];
 				else
-					return (FileItem)listBox.SelectedItems[listBox.SelectedItems.Count - 1];
+					return (FileItem)TargetListBox.SelectedItems[TargetListBox.SelectedItems.Count - 1];
 			});
 
-			var moveFunc = new Action<int>((int steps) =>
+			Action<int> MoveFunc = new Action<int>((int Steps) =>
 			{
-				FileItem selectedItem = getSelectedItem();
-				var index = listBox.Items.IndexOf(selectedItem);
-				var newIndex = Math.Min(Math.Max(0, index + steps), listBox.Items.Count - 1);
-				if (newIndex == index)
+				FileItem SelectedItem = GetSelectedItem();
+				int Index = TargetListBox.Items.IndexOf(SelectedItem);
+				int NewIndex = Math.Min(Math.Max(0, Index + Steps), TargetListBox.Items.Count - 1);
+				if (NewIndex == Index)
+				{
 					return;
-				selectedItem = (FileItem)listBox.Items[newIndex];
-				if (selectedItem == null)
+				}
+				SelectedItem = (FileItem)TargetListBox.Items[NewIndex];
+				if (SelectedItem == null)
+				{
 					return;
-				listBox.SelectedItems.Clear();
-				listBox.SelectedItems.Add(selectedItem);
-				listBox.ScrollIntoView(selectedItem);
+				}
+				TargetListBox.SelectedItems.Clear();
+				TargetListBox.SelectedItems.Add(SelectedItem);
+				TargetListBox.ScrollIntoView(SelectedItem);
 			});
 
-			if (e.Key == Key.Down)
+			if (E.Key == Key.Down)
 			{
-				moveFunc(1);
+				MoveFunc(1);
 				return;
 			}
-			if (e.Key == Key.Up)
+			if (E.Key == Key.Up)
 			{
-				moveFunc(-1);
+				MoveFunc(-1);
 				return;
 			}
-			if (e.Key == Key.Enter)
+			if (E.Key == Key.Enter)
 			{
 				if (OpenSelectedFile())
 					HidePanel();
 				return;
 			}
-			if (e.Key == Key.PageDown)
+			if (E.Key == Key.PageDown)
 			{
-				e.Handled = true;
-				var lastVisible = GetVisibleListViewElement(listBox, (int)listBox.ActualHeight - 22);
-				if (lastVisible == null)
+				E.Handled = true;
+				ListBoxItem LastVisible = GetVisibleListViewElement(TargetListBox, (int)TargetListBox.ActualHeight - 22);
+				if (LastVisible == null)
+				{
 					return;
-				var lastItem = listBox.ItemContainerGenerator.ItemFromContainer(lastVisible);
-				FileItem selectedItem = getSelectedItem();
-				if (lastItem != selectedItem)
+				}
+				object LastItem = TargetListBox.ItemContainerGenerator.ItemFromContainer(LastVisible);
+				FileItem SelectedItem = GetSelectedItem();
+				if (LastItem != SelectedItem)
 				{
-					listBox.SelectedItems.Clear();
-					listBox.SelectedItems.Add(lastItem);
-					listBox.ScrollIntoView(lastItem);
+					TargetListBox.SelectedItems.Clear();
+					TargetListBox.SelectedItems.Add(LastItem);
+					TargetListBox.ScrollIntoView(LastItem);
 				}
 				else
-					moveFunc(((int)(listBox.ActualHeight / lastVisible.ActualHeight)) - 1);
+				{
+					MoveFunc(((int)(TargetListBox.ActualHeight / LastVisible.ActualHeight)) - 1);
+				}
 				return;
 			}
-			if (e.Key == Key.PageUp)
+			if (E.Key == Key.PageUp)
 			{
-				e.Handled = true;
-				var firstVisible = GetVisibleListViewElement(listBox, 10);
-				if (firstVisible == null)
+				E.Handled = true;
+				ListBoxItem FirstVisible = GetVisibleListViewElement(TargetListBox, 10);
+				if (FirstVisible == null)
+				{
 					return;
-				var firstItem = listBox.ItemContainerGenerator.ItemFromContainer(firstVisible);
-				FileItem selectedItem = getSelectedItem();
-				if (firstItem != selectedItem)
+				}
+				object FirstItem = TargetListBox.ItemContainerGenerator.ItemFromContainer(FirstVisible);
+				FileItem SelectedItem = GetSelectedItem();
+				if (FirstItem != SelectedItem)
 				{
-					listBox.SelectedItems.Clear();
-					listBox.SelectedItems.Add(firstItem);
-					listBox.ScrollIntoView(firstItem);
+					TargetListBox.SelectedItems.Clear();
+					TargetListBox.SelectedItems.Add(FirstItem);
+					TargetListBox.ScrollIntoView(FirstItem);
 				}
 				else
-					moveFunc(1 - (int)(listBox.ActualHeight / firstVisible.ActualHeight));
+				{
+					MoveFunc(1 - (int)(TargetListBox.ActualHeight / FirstVisible.ActualHeight));
+				}
 			}
 		}
 
-		private void FilterEditBox_TextChanged(object sender, TextChangedEventArgs e)
+		private void FilterEditBox_TextChanged(object Sender, TextChangedEventArgs E)
 		{
 			RefreshListViews();
 			//AllFilesListBox.SelectedIndex = 0;
 		}
 
-		private void FileListView_GotFocus(object sender, RoutedEventArgs e)
+		private void FileListView_GotFocus(object Sender, RoutedEventArgs E)
 		{
-			ActiveFilesListBox = (ListBox)e.Source;
+			ActiveFilesListBox = (ListBox)E.Source;
 		}
 
-		private void FileListView_SelectionChanged(object sender, SelectionChangedEventArgs e)
+		private void FileListView_SelectionChanged(object Sender, SelectionChangedEventArgs E)
 		{
 			RefreshStatusBox();
 		}
 
-		private void FileListView_KeyDown(object sender, KeyEventArgs e)
+		private void FileListView_KeyDown(object Sender, KeyEventArgs E)
 		{
 			if ((Keyboard.Modifiers & ModifierKeys.Control) != 0)
 			{
-				if (e.Key == Key.I)
+				if (E.Key == Key.I)
 				{
 					if (!CreateIncludePath())
 					{
 						return;
 					}
-					
+
 					HidePanel();
 
 					if ((Keyboard.Modifiers & ModifierKeys.Shift) == 0)
@@ -254,140 +281,152 @@ namespace UnrealVS
 					}
 
 					// Let's add directly in to active document
-					var activeDoc = UnrealVSPackage.Instance.DTE.ActiveDocument;
-					var doc = (TextDocument)(activeDoc.Object("TextDocument"));
-					if (doc == null)
+					Document ActiveDoc = UnrealVSPackage.Instance.DTE.ActiveDocument;
+					TextDocument Doc = (TextDocument)(ActiveDoc.Object("TextDocument"));
+					if (Doc == null)
+					{
 						return;
+					}
 
-					bool addSpace = false;
-					var includePath = GetIncludePath((FileItem)(ActiveFilesListBox.SelectedItems[0]));
-					string pasteString = $"#include \"{includePath}\"";
+					bool bAddSpace = false;
+					string IncludePath = GetIncludePath((FileItem)(ActiveFilesListBox.SelectedItems[0]));
+					string PasteString = $"#include \"{IncludePath}\"";
 
-					bool skipFirstInclude = activeDoc.FullName.EndsWith(".cpp");
-					int addBeforeLine = -1;
-					var lastIncludeLine = -1;
+					bool bSkipFirstInclude = ActiveDoc.FullName.EndsWith(".cpp");
+					int AddBeforeLine = -1;
+					int LastIncludeLine = -1;
 
-					bool isInComment = false;
-					var p = doc.StartPoint.CreateEditPoint();
-					var lastLine = doc.EndPoint.Line;
-					for (int lineIndex = 1; lineIndex < lastLine; ++lineIndex)
+					bool bIsInComment = false;
+					EditPoint EditPoint = Doc.StartPoint.CreateEditPoint();
+					int LastLine = Doc.EndPoint.Line;
+					for (int LineIndex = 1; LineIndex < LastLine; ++LineIndex)
 					{
-						var str = p.GetLines(lineIndex, lineIndex + 1).Trim();
-						if (str.Length == 0)
+						string Str = EditPoint.GetLines(LineIndex, LineIndex + 1).Trim();
+						if (Str.Length == 0)
 						{
 							continue;
 						}
 
-						if (isInComment)
+						if (bIsInComment)
 						{
-							if (str.IndexOf("*/") != -1)
-								isInComment = false;
+							if (Str.IndexOf("*/") != -1)
+							{
+								bIsInComment = false;
+							}
 							continue;
 						}
 
-						if (str.StartsWith("//"))
+						if (Str.StartsWith("//"))
 						{
 							continue;
 						}
 
-						if (str.StartsWith("/*"))
+						if (Str.StartsWith("/*"))
 						{
-							isInComment = true;
+							bIsInComment = true;
 							continue;
 						}
 
-						if (str.StartsWith("#include"))
+						if (Str.StartsWith("#include"))
 						{
-							if (str == pasteString)
+							if (Str == PasteString)
 							{
-								doc.Selection.MoveTo(lineIndex, 0);
-								doc.Selection.SelectLine();
-								p.TryToShow();
+								Doc.Selection.MoveTo(LineIndex, 0);
+								Doc.Selection.SelectLine();
+								EditPoint.TryToShow();
 								return;
 							}
 
-							bool isFirst = lastIncludeLine == -1;
-							int prevLastIncludeLine = lastIncludeLine;
-							lastIncludeLine = lineIndex;
-							if (skipFirstInclude && isFirst)
+							bool bIsFirst = LastIncludeLine == -1;
+							int PrevLastIncludeLine = LastIncludeLine;
+							LastIncludeLine = LineIndex;
+							if (bSkipFirstInclude && bIsFirst)
 							{
 								continue;
 							}
 
-							var span = str.AsSpan(9);
-							var firstQuote = span.IndexOf('"');
-							if (firstQuote != -1)
+							ReadOnlySpan<char> Span = Str.AsSpan(9);
+							int FirstQuote = Span.IndexOf('"');
+							if (FirstQuote != -1)
 							{
-								span = span.Slice(firstQuote + 1);
-								var secondQuote = span.IndexOf('"');
-								if (secondQuote != -1)
+								Span = Span.Slice(FirstQuote + 1);
+								int SecondQuote = Span.IndexOf('"');
+								if (SecondQuote != -1)
 								{
-									span = span.Slice(0, secondQuote).Trim();
-									bool isGeneratedInclude = span.IndexOf(".generated.".AsSpan()) != -1;
-									if (includePath.AsSpan().CompareTo(span, StringComparison.Ordinal) < 0 || isGeneratedInclude)
+									Span = Span.Slice(0, SecondQuote).Trim();
+									bool bIsGeneratedInclude = Span.IndexOf(".generated.".AsSpan()) != -1;
+									if (IncludePath.AsSpan().CompareTo(Span, StringComparison.Ordinal) < 0 || bIsGeneratedInclude)
 									{
-										if (isGeneratedInclude && prevLastIncludeLine != -1)
-											lastIncludeLine = prevLastIncludeLine;
+										if (bIsGeneratedInclude && PrevLastIncludeLine != -1)
+										{
+											LastIncludeLine = PrevLastIncludeLine;
+										}
 										else
-											addBeforeLine = lineIndex;
+										{
+											AddBeforeLine = LineIndex;
+										}
 										break;
 									}
 								}
 							}
-							else if (span.IndexOf("UE_INLINE_GENERATED_CPP".AsSpan()) != -1)
+							else if (Span.IndexOf("UE_INLINE_GENERATED_CPP".AsSpan()) != -1)
 							{
-								lastIncludeLine = prevLastIncludeLine;
+								LastIncludeLine = PrevLastIncludeLine;
 								break;
 							}
 							continue;
 						}
-							
-						if (str.StartsWith("#pragma once"))
+
+						if (Str.StartsWith("#pragma once"))
 						{
-							lastIncludeLine = lineIndex + 1;
+							LastIncludeLine = LineIndex + 1;
 							continue;
 						}
 
-						lastIncludeLine = lineIndex - 1;
-						addSpace = true;
+						LastIncludeLine = LineIndex - 1;
+						bAddSpace = true;
 						break;  // Unknown code
 					}
 
-					if (addBeforeLine == -1 && lastIncludeLine != -1)
+					if (AddBeforeLine == -1 && LastIncludeLine != -1)
 					{
-						addBeforeLine = lastIncludeLine + 1;
+						AddBeforeLine = LastIncludeLine + 1;
 					}
 
-					if (addBeforeLine == -1)
+					if (AddBeforeLine == -1)
 					{
-						addBeforeLine = lastLine;
+						AddBeforeLine = LastLine;
 					}
 
-					p.LineDown(addBeforeLine - 1);
-					p.Insert($"{pasteString}\r\n");
-					if (addSpace)
+					EditPoint.LineDown(AddBeforeLine - 1);
+					EditPoint.Insert($"{PasteString}\r\n");
+					if (bAddSpace)
 					{
-						p.Insert("\r\n");
+						EditPoint.Insert("\r\n");
 					}
-					doc.Selection.MoveTo(addBeforeLine, 0);
-					doc.Selection.SelectLine();
-					p.TryToShow();
+					Doc.Selection.MoveTo(AddBeforeLine, 0);
+					Doc.Selection.SelectLine();
+					EditPoint.TryToShow();
 				}
 				return;
 			}
-			if (e.Key == Key.Left)
+			if (E.Key == Key.Left)
 			{
 				if (ActiveFilesListBox != AllFilesListBox)
+				{
 					--FilesListTab.SelectedIndex;
+				}
 				//if (ActiveFilesListBox == BookmarkedFilesListBox)
 				//    ToggleListView(BookmarkedFilesListBox, AllFilesListBox);
 				return;
 			}
 
-			if (e.Key == Key.Right)
+			if (E.Key == Key.Right)
 			{
 				if (ActiveFilesListBox != RecentFilesListBox)// && ((List<FileItem>)BookmarkedFilesListBox.ItemsSource).Count > 0)
+				{
 					++FilesListTab.SelectedIndex;
+				}
 				//if (ActiveFilesListBox == AllFilesListBox)// && ((List<FileItem>)BookmarkedFilesListBox.ItemsSource).Count > 0)
 				//   FilesListTab.SelectedIndex = 1;
 				//if (ActiveFilesListBox == AllFilesListBox && ((List<FileItem>)BookmarkedFilesListBox.ItemsSource).Count > 0)
@@ -395,30 +434,38 @@ namespace UnrealVS
 				return;
 			}
 
-			if (e.Key == Key.Tab)
+			if (E.Key == Key.Tab)
 			{
-				e.Handled = true;
+				E.Handled = true;
 				FilterEditBox.Focus();
 				return;
 			}
 
-			if (e.Key == Key.Insert)
+			if (E.Key == Key.Insert)
 			{
 				if (ActiveFilesListBox != AllFilesListBox)
+				{
 					return;
-				bool BookmarksModified = false;
-				foreach (var i in ActiveFilesListBox.SelectedItems)
+				}
+				bool bBookmarksModified = false;
+				foreach (object SelectedItem in ActiveFilesListBox.SelectedItems)
 				{
-					var item = (FileItem)i;
-					if (item == null || item.Name == "")
+					FileItem Item = (FileItem)SelectedItem;
+					if (Item == null || Item.Name == "")
+					{
 						continue;
-					if (BookmarkedFileItems.Contains(item))
+					}
+					if (BookmarkedFileItems.Contains(Item))
+					{
 						continue;
-					BookmarkedFileItems.Add(item);
-					BookmarksModified = true;
+					}
+					BookmarkedFileItems.Add(Item);
+					bBookmarksModified = true;
 				}
-				if (!BookmarksModified)
+				if (!bBookmarksModified)
+				{
 					return;
+				}
 				RefreshListView(BookmarkedFilesListBox, BookmarkedFileItems.ToArray());
 				//if (AllFilesListBox.SelectedIndex < AllFilesListBox.Items.Count - 1)
 				//    FocusListViewItem(AllFilesListBox, AllFilesListBox.SelectedIndex + 1);
@@ -426,78 +473,90 @@ namespace UnrealVS
 				return;
 			}
 
-			if (e.Key == Key.Delete)
+			if (E.Key == Key.Delete)
 			{
 				if (ActiveFilesListBox != BookmarkedFilesListBox)
+				{
 					return;
-				bool BookmarksModified = false;
+				}
+				bool bBookmarksModified = false;
 
-				var index = BookmarkedFilesListBox.SelectedIndex;
-				foreach (var i in ActiveFilesListBox.SelectedItems)
+				int Index = BookmarkedFilesListBox.SelectedIndex;
+				foreach (object SelectedItem in ActiveFilesListBox.SelectedItems)
 				{
-					var item = (FileItem)i;
-					if (item == null || item.Name == "")
+					FileItem Item = (FileItem)SelectedItem;
+					if (Item == null || Item.Name == "")
+					{
 						continue;
-					if (!BookmarkedFileItems.Remove(item))
+					}
+					if (!BookmarkedFileItems.Remove(Item))
+					{
 						continue;
-					BookmarksModified = true;
+					}
+					bBookmarksModified = true;
 				}
-				if (BookmarksModified)
+				if (bBookmarksModified)
 				{
 					RefreshListView(BookmarkedFilesListBox, BookmarkedFileItems.ToArray());
-					if (index == BookmarkedFilesListBox.Items.Count)
+					if (Index == BookmarkedFilesListBox.Items.Count)
 					{
-						if (index == 0)
+						if (Index == 0)
+						{
 							FilesListTab.SelectedIndex = 0;//ActiveFilesListBox = AllFilesListBox;
+						}
 						else
-							index -= 1;
+						{
+							Index -= 1;
+						}
 					}
-					BookmarkedFilesListBox.SelectedIndex = index;
+					BookmarkedFilesListBox.SelectedIndex = Index;
 					AsyncFocusSelectedItem();
 					SaveSolutionSettings();
 				}
 				return;
 			}
 
-			if (e.Key == Key.Enter)
+			if (E.Key == Key.Enter)
 			{
 				if (OpenSelectedFile())
 					HidePanel();
 				return;
 			}
 
-			bool resetSelection = false;
+			bool bResetSelection = false;
 
-			if (e.Key == Key.Back)
+			if (E.Key == Key.Back)
 			{
-				string str = FilterEditBox.Text;
-				if (string.IsNullOrEmpty(str))
+				string Str = FilterEditBox.Text;
+				if (string.IsNullOrEmpty(Str))
 					return;
-				if (ClearFilter)
+				if (bClearFilter)
 				{
-					resetSelection = true;
-					ClearFilter = false;
+					bResetSelection = true;
+					bClearFilter = false;
 					FilterEditBox.Text = "";
 				}
 				else
-					FilterEditBox.Text = str.Substring(0, str.Length - 1);
+				{
+					FilterEditBox.Text = Str.Substring(0, Str.Length - 1);
+				}
 			}
 			else
 			{
-				char c = GetCharFromKey(e.Key);
-				if (c == 0)
+				char Ch = GetCharFromKey(E.Key);
+				if (Ch == 0)
 					return;
-				if (ClearFilter)
+				if (bClearFilter)
 				{
-					resetSelection = true;
-					ClearFilter = false;
+					bResetSelection = true;
+					bClearFilter = false;
 					FilterEditBox.Text = "";
 				}
-				FilterEditBox.Text += c;
+				FilterEditBox.Text += Ch;
 				FilterEditBox.SelectAll();
 			}
 
-			if (resetSelection)
+			if (bResetSelection)
 			{
 				AllFilesListBox.SelectedIndex = 0;
 				BookmarkedFilesListBox.SelectedIndex = 0;
@@ -506,10 +565,10 @@ namespace UnrealVS
 
 			AsyncFocusSelectedItem();
 
-			e.Handled = true;
+			E.Handled = true;
 		}
 
-		private void ListView_MouseDoubleClick(object sender, MouseButtonEventArgs e)
+		private void ListView_MouseDoubleClick(object Sender, MouseButtonEventArgs E)
 		{
 			if (OpenSelectedFile())
 				HidePanel();
@@ -517,107 +576,125 @@ namespace UnrealVS
 
 		private bool OpenSelectedFile()
 		{
-			var toOpen = new List<string>();
-			foreach (var i in ActiveFilesListBox.SelectedItems)
+			List<string> ToOpen = new List<string>();
+			foreach (object SelectedItem in ActiveFilesListBox.SelectedItems)
 			{
-				var item = (FileItem)i;
-				if (item == null || item.Name == "")
+				FileItem Item = (FileItem)SelectedItem;
+				if (Item == null || Item.Name == "")
+				{
 					continue;
-				toOpen.Add(item.File);
+				}
+				ToOpen.Add(Item.File);
 			}
 
 			ThreadHelper.JoinableTaskFactory.Run(async delegate
 			{
 				await ThreadHelper.JoinableTaskFactory.SwitchToMainThreadAsync();
-				foreach (var f in toOpen)
+				foreach (string F in ToOpen)
 				{
-					VsShellUtilities.OpenDocument(UnrealVSPackage.Instance, f);
+					VsShellUtilities.OpenDocument(UnrealVSPackage.Instance, F);
 				}
 			});
 
-			return toOpen.Count != 0;
+			return ToOpen.Count != 0;
 		}
 
-		private string GetIncludePath(FileItem item)
+		private string GetIncludePath(FileItem Item)
 		{
-                if (item == null || item.Name == "")
+			if (Item == null || Item.Name == "")
+			{
 				return null;
-                string includePath = item.File;
-                includePath = includePath.Replace("\\", "/");
-                string includePathLwr = includePath.ToLower();
-			var includeStartIndex = Math.Max(Math.Max(includePathLwr.LastIndexOf("/public/"), includePathLwr.LastIndexOf("/private/")), includePathLwr.LastIndexOf("/classes/"));
-                if (includeStartIndex != -1)
-                    includePath = includePath.Substring(includePath.IndexOf('/', includeStartIndex + 1) + 1);
-			return includePath;
+			}
+			string IncludePath = Item.File;
+			IncludePath = IncludePath.Replace("\\", "/");
+			string IncludePathLwr = IncludePath.ToLower();
+			int IncludeStartIndex = Math.Max(Math.Max(IncludePathLwr.LastIndexOf("/public/"), IncludePathLwr.LastIndexOf("/private/")), IncludePathLwr.LastIndexOf("/classes/"));
+			if (IncludeStartIndex != -1)
+			{
+				IncludePath = IncludePath.Substring(IncludePath.IndexOf('/', IncludeStartIndex + 1) + 1);
+			}
+			return IncludePath;
 		}
 
 		private bool CreateIncludePath()
 		{
-			string includeStrings = "";
-			foreach (var i in ActiveFilesListBox.SelectedItems)
+			string IncludeStrings = "";
+			foreach (object SelectedItem in ActiveFilesListBox.SelectedItems)
 			{
-				var includePath = GetIncludePath((FileItem)i);
-				if (includePath != null)
-                includeStrings += $"#include \"{includePath}\"\r\n";
+				string IncludePath = GetIncludePath((FileItem)SelectedItem);
+				if (IncludePath != null)
+				{
+					IncludeStrings += $"#include \"{IncludePath}\"\r\n";
+				}
 			}
-			if (includeStrings == "")
+			if (IncludeStrings == "")
 				return false;
-			Clipboard.SetText(includeStrings);
+			Clipboard.SetText(IncludeStrings);
 			return true;
 		}
 
-		private void ToggleListView(ListBox from, ListBox to)
+		private void ToggleListView(ListBox From, ListBox To)
 		{
-			var firstVisibleFrom = GetVisibleListViewElement(from, 10);
-			if (firstVisibleFrom == null)
+			ListBoxItem FirstVisibleFrom = GetVisibleListViewElement(From, 10);
+			if (FirstVisibleFrom == null)
+			{
 				return;
-			var firstVisibleFromIndex = from.ItemContainerGenerator.IndexFromContainer(firstVisibleFrom);
-			var firstVisibleTo = GetVisibleListViewElement(to, 10);
-			if (firstVisibleTo == null)
+			}
+			int FirstVisibleFromIndex = From.ItemContainerGenerator.IndexFromContainer(FirstVisibleFrom);
+			ListBoxItem FirstVisibleTo = GetVisibleListViewElement(To, 10);
+			if (FirstVisibleTo == null)
+			{
 				return;
-			var firstVisibleToIndex = to.ItemContainerGenerator.IndexFromContainer(firstVisibleTo);
-			var toSelectIndex = firstVisibleToIndex + from.SelectedIndex - firstVisibleFromIndex;
-			if (toSelectIndex >= to.Items.Count)
-				toSelectIndex = to.Items.Count - 1;
-			FocusListViewItem(to, toSelectIndex);
+			}
+			int FirstVisibleToIndex = To.ItemContainerGenerator.IndexFromContainer(FirstVisibleTo);
+			int ToSelectIndex = FirstVisibleToIndex + From.SelectedIndex - FirstVisibleFromIndex;
+			if (ToSelectIndex >= To.Items.Count)
+				ToSelectIndex = To.Items.Count - 1;
+			FocusListViewItem(To, ToSelectIndex);
 		}
 
-		private void FocusListViewItem(ListBox listView, int itemIndex)
+		private void FocusListViewItem(ListBox ListView, int ItemIndex)
 		{
-			if (itemIndex == -1)
-				itemIndex = 0;
-			listView.SelectedIndex = itemIndex;
-			if (listView.Items.Count > 0)
-				(listView.ItemContainerGenerator.ContainerFromIndex(itemIndex) as ListBoxItem)?.Focus();
+			if (ItemIndex == -1)
+			{
+				ItemIndex = 0;
+			}
+			ListView.SelectedIndex = ItemIndex;
+			if (ListView.Items.Count > 0)
+			{
+				(ListView.ItemContainerGenerator.ContainerFromIndex(ItemIndex) as ListBoxItem)?.Focus();
+			}
 			else
-				listView.Focus();
+			{
+				ListView.Focus();
+			}
 			RefreshStatusBox();
 		}
 
 		private void AsyncFocusSelectedItem()
 		{
-			double interval = 0.1;
-			var timer = new DispatcherTimer(DispatcherPriority.Normal);
-			timer.Tick += (s, e) =>
+			double Interval = 0.1;
+			DispatcherTimer Timer = new DispatcherTimer(DispatcherPriority.Normal);
+			Timer.Tick += (S, E) =>
 			{
-				var items = (List<FileItem>)ActiveFilesListBox.ItemsSource;
-				if (items.Count == 0)
+				List<FileItem> Items = (List<FileItem>)ActiveFilesListBox.ItemsSource;
+				if (Items.Count == 0)
 				{
 					ActiveFilesListBox.Focus();
-					timer.Stop();
+					Timer.Stop();
 					return;
 				}
-				var item = (ActiveFilesListBox.ItemContainerGenerator.ContainerFromIndex(Math.Max(0, ActiveFilesListBox.SelectedIndex)) as ListBoxItem);
-				if (item == null)
+				ListBoxItem Item = (ActiveFilesListBox.ItemContainerGenerator.ContainerFromIndex(Math.Max(0, ActiveFilesListBox.SelectedIndex)) as ListBoxItem);
+				if (Item == null)
 				{
-					timer.Interval = TimeSpan.FromSeconds(interval);
-					interval = Math.Min(interval + 0.1, 1.0);
+					Timer.Interval = TimeSpan.FromSeconds(Interval);
+					Interval = Math.Min(Interval + 0.1, 1.0);
 					return;
 				}
-				item.Focus();
-				timer.Stop();
+				Item.Focus();
+				Timer.Stop();
 			};
-			timer.Start();
+			Timer.Start();
 		}
 
 		private void RefreshListViews()
@@ -627,140 +704,140 @@ namespace UnrealVS
 			RefreshListView(RecentFilesListBox, RecentFileItems.ToArray());
 		}
 
-		struct Score
+		private struct Score
 		{
 			public int Sorting;
 			public int Index;
 		}
 
-		private void RefreshListView(ListBox listView, FileItem[] source)
+		private void RefreshListView(ListBox ListView, FileItem[] Source)
 		{
-			var fileItems = (List<FileItem>)listView.ItemsSource;
-			fileItems.Clear();
+			List<FileItem> FileItems = (List<FileItem>)ListView.ItemsSource;
+			FileItems.Clear();
 
-			string filterStr = FilterEditBox.Text;
-			var filter = filterStr.Split(new char[] { ' ' }, StringSplitOptions.RemoveEmptyEntries);
+			string FilterStr = FilterEditBox.Text;
+			string[] Filter = FilterStr.Split(new char[] { ' ' }, StringSplitOptions.RemoveEmptyEntries);
 
-			var selections = new HashSet<FileItem>();
-			foreach (var i in listView.SelectedItems)
+			HashSet<FileItem> Selections = new HashSet<FileItem>();
+			foreach (object SelectedItem in ListView.SelectedItems)
 			{
-				var item = (FileItem)i;
-				if (item != null && item.Name != "")
-					selections.Add(item);
+				FileItem Item = (FileItem)SelectedItem;
+				if (Item != null && Item.Name != "")
+					Selections.Add(Item);
 			}
 
-			var filterLen = filter.Length;
-			if (filterLen > 0)
+			int FilterLen = Filter.Length;
+			if (FilterLen > 0)
 			{
-				var filterSortedItems = new List<ValueTuple<Score, FileItem>>();
-				var scoreArray = new ValueTuple<int, int>[5];
+				List<ValueTuple<Score, FileItem>> FilterSortedItems = new List<ValueTuple<Score, FileItem>>();
+				ValueTuple<int, int>[] ScoreArray = new ValueTuple<int, int>[5];
 				int ItemIndex = 0;
-				foreach (var item in source)
+				foreach (FileItem Item in Source)
 				{
-					int scoreArrayIndex = 0;
-					for (int i = 0; i != filterLen; ++i)
+					int ScoreArrayIndex = 0;
+					for (int i = 0; i != FilterLen; ++i)
 					{
-						var f = filter[i];
-						int index = item.Name.IndexOf(f, System.StringComparison.OrdinalIgnoreCase);
-						if (index == -1)
+						string F = Filter[i];
+						int Index = Item.Name.IndexOf(F, System.StringComparison.OrdinalIgnoreCase);
+						if (Index == -1)
 						{
-							scoreArrayIndex = 0;
+							ScoreArrayIndex = 0;
 							break;
 						}
-						scoreArray[scoreArrayIndex].Item1 = index;
-						scoreArray[scoreArrayIndex].Item2 = scoreArrayIndex;
-						++scoreArrayIndex;
-						if (scoreArrayIndex == 5)
+						ScoreArray[ScoreArrayIndex].Item1 = Index;
+						ScoreArray[ScoreArrayIndex].Item2 = ScoreArrayIndex;
+						++ScoreArrayIndex;
+						if (ScoreArrayIndex == 5)
 							break;
 					}
 
-					if (scoreArrayIndex == 0)
+					if (ScoreArrayIndex == 0)
 						continue;
 
 					// We want to sort the ones that match the filter orders first and prioritize the ones that starts with the first filter entry
-					int sortingScore = 0;
-					Array.Sort(scoreArray, 0, scoreArrayIndex);
-					for (int i = 0; i != scoreArrayIndex; ++i)
+					int SortingScore = 0;
+					Array.Sort(ScoreArray, 0, ScoreArrayIndex);
+					for (int i = 0; i != ScoreArrayIndex; ++i)
 					{
-						var Tuple = scoreArray[i];
-						if (i == 0 && Tuple.Item1 == 0 && Tuple.Item2 == 0)
-							sortingScore -= 100;
-						sortingScore += (10 >> Tuple.Item2) * i + Tuple.Item1;
+						ValueTuple<int, int> ScoreEntry = ScoreArray[i];
+						if (i == 0 && ScoreEntry.Item1 == 0 && ScoreEntry.Item2 == 0)
+							SortingScore -= 100;
+						SortingScore += (10 >> ScoreEntry.Item2) * i + ScoreEntry.Item1;
 					}
-					var s = new Score() { Sorting = sortingScore, Index = ItemIndex++ };
-					filterSortedItems.Add(new ValueTuple<Score, FileItem>(s, item));
+					Score NewScore = new Score() { Sorting = SortingScore, Index = ItemIndex++ };
+					FilterSortedItems.Add(new ValueTuple<Score, FileItem>(NewScore, Item));
 				}
 
-				filterSortedItems.Sort((a, b) =>
+				FilterSortedItems.Sort((A, B) =>
 				{
-					if (a.Item1.Sorting != b.Item1.Sorting)
-						return a.Item1.Sorting - b.Item1.Sorting;
-					return a.Item1.Index - b.Item1.Index;
+					if (A.Item1.Sorting != B.Item1.Sorting)
+						return A.Item1.Sorting - B.Item1.Sorting;
+					return A.Item1.Index - B.Item1.Index;
 				});
-				foreach (var tuple in filterSortedItems)
-					fileItems.Add(tuple.Item2);
+				foreach (ValueTuple<Score, FileItem> SortedItem in FilterSortedItems)
+					FileItems.Add(SortedItem.Item2);
 			}
 			else
 			{
-				fileItems.AddRange(source);
+				FileItems.AddRange(Source);
 			}
 
-			int fileCount = fileItems.Count;
+			int FileCount = FileItems.Count;
 
 			// If empty we need to add one entry to make sure listview can be focused in windows
-			if (fileItems.Count == 0)
-				fileItems.Add(new FileItem() { Name = "" });
+			if (FileItems.Count == 0)
+				FileItems.Add(new FileItem() { Name = "" });
 
-			CollectionViewSource.GetDefaultView(fileItems)?.Refresh();
+			CollectionViewSource.GetDefaultView(FileItems)?.Refresh();
 
-			if (selections.Count != 0)
+			if (Selections.Count != 0)
 			{
-				foreach (var item in fileItems)
+				foreach (FileItem Item in FileItems)
 				{
-					if (!selections.Contains(item))
+					if (!Selections.Contains(Item))
 						continue;
-					listView.SelectedItems.Add(item);
-					listView.ScrollIntoView(item);
+					ListView.SelectedItems.Add(Item);
+					ListView.ScrollIntoView(Item);
 				}
 			}
 
-			if (listView.SelectedItems.Count == 0)
-				listView.SelectedItems.Add(fileItems[0]);
+			if (ListView.SelectedItems.Count == 0)
+				ListView.SelectedItems.Add(FileItems[0]);
 
 
-			if (listView == AllFilesListBox)
-				AllFilesTab.Header = $"All Files ({fileCount}/{source.Length})";
-			else if (listView == BookmarkedFilesListBox)
-				BookmarkedFilesTab.Header = $"Bookmarked Files ({fileCount}/{source.Length})";
+			if (ListView == AllFilesListBox)
+				AllFilesTab.Header = $"All Files ({FileCount}/{Source.Length})";
+			else if (ListView == BookmarkedFilesListBox)
+				BookmarkedFilesTab.Header = $"Bookmarked Files ({FileCount}/{Source.Length})";
 			else
-				RecentFilesTab.Header = $"Recent Files ({fileCount}/{source.Length})";
+				RecentFilesTab.Header = $"Recent Files ({FileCount}/{Source.Length})";
 		}
 
 		private void RefreshStatusBox()
 		{
-			var items = ActiveFilesListBox.SelectedItems;
-			if (items.Count == 0)
+			System.Collections.IList Items = ActiveFilesListBox.SelectedItems;
+			if (Items.Count == 0)
 				StatusText.Text = "";
-			else if (items.Count == 1)
+			else if (Items.Count == 1)
 			{
-				var item = (FileItem)items[0];
-				StatusText.Text = item.File;
+				FileItem Item = (FileItem)Items[0];
+				StatusText.Text = Item.File;
 			}
 			else
 				StatusText.Text = "Multiple files selected";
 		}
 
-		private void FileListView_PreviewKeyDown(object sender, KeyEventArgs e)
+		private void FileListView_PreviewKeyDown(object Sender, KeyEventArgs E)
 		{
-			if (e.Key != Key.Space)
+			if (E.Key != Key.Space)
 				return;
-			if (ClearFilter)
+			if (bClearFilter)
 			{
-				ClearFilter = false;
+				bClearFilter = false;
 				return;
 			}
 			FilterEditBox.Text += ' ';
-			e.Handled = true;
+			E.Handled = true;
 		}
 
 		internal void HandleEscape()
@@ -785,58 +862,66 @@ namespace UnrealVS
 		{
 			SaveSolutionSettings();
 			if (IsVisible)
+			{
 				AsyncRefreshFileList();
+			}
 			else
+			{
 				AllFileItems = Array.Empty<FileItem>(); // Will trigger AsyncRefreshFileList when visible
+			}
 		}
 
 		internal void HandleDocumentActivated(Document Document)
 		{
-			var fileName = Document.FullName;
-			FileItem fileItem;
-			if (AllFileItemsLookup == null || !AllFileItemsLookup.TryGetValue(fileName, out fileItem))
+			string FileName = Document.FullName;
+			FileItem FileItem;
+			if (AllFileItemsLookup == null || !AllFileItemsLookup.TryGetValue(FileName, out FileItem))
+			{
 				return;
+			}
 
-			RecentFileItems.RemoveAll((item) => item == fileItem);
-			RecentFileItems.Insert(0, fileItem);
+			RecentFileItems.RemoveAll((Item) => Item == FileItem);
+			RecentFileItems.Insert(0, FileItem);
 
 			int RecentMaxCount = 100;
 			if (RecentFileItems.Count > RecentMaxCount)
+			{
 				RecentFileItems.RemoveRange(RecentMaxCount, RecentFileItems.Count - RecentMaxCount);
+			}
 
 			RefreshListView(RecentFilesListBox, RecentFileItems.ToArray());
 		}
 
 		private void AsyncRefreshFileList()
 		{
-			if (IsRefreshing)
+			if (bIsRefreshing)
 				return;
-			IsRefreshing = true;
+			bIsRefreshing = true;
 			RefreshingText.Text = "Refreshing File Lists (0)";
 			RefreshingDialog.Visibility = Visibility.Visible;
 
-			DispatcherTimer timer = new DispatcherTimer(DispatcherPriority.Render);
-			timer.Interval = TimeSpan.FromSeconds(0.2);
-			var traverser = new SolutionTraverser();
-			timer.Tick += (s, e_) =>
+			DispatcherTimer Timer = new DispatcherTimer(DispatcherPriority.Render);
+			Timer.Interval = TimeSpan.FromSeconds(0.2);
+			SolutionTraverser Traverser = new SolutionTraverser();
+			Timer.Tick += (S, E) =>
 			{
-				var listItems = traverser.Update();
-				if (listItems == null)
+				FileItem[] ListItems = Traverser.Update();
+				if (ListItems == null)
 				{
-					RefreshingText.Text = $"Refreshing File Lists ({traverser.HandledItemCount})";
-					timer.Interval = TimeSpan.FromSeconds(0.01);
+					RefreshingText.Text = $"Refreshing File Lists ({Traverser.HandledItemCount})";
+					Timer.Interval = TimeSpan.FromSeconds(0.01);
 					return;
 				}
-				AllFileItems = listItems;
+				AllFileItems = ListItems;
 
-				var lookup = new Dictionary<string, FileItem>();
-				foreach (var item in AllFileItems)
-					try { lookup.Add(item.File, item); } catch (Exception) { }
-				AllFileItemsLookup = lookup;
+				Dictionary<string, FileItem> Lookup = new Dictionary<string, FileItem>();
+				foreach (FileItem Item in AllFileItems)
+					try { Lookup.Add(Item.File, Item); } catch (Exception) { }
+				AllFileItemsLookup = Lookup;
 
 				LoadSolutionSettings();
 
-				//ClearFilter = true;
+				//bClearFilter = true;
 				RefreshListViews();
 
 				//AllFilesListBox.SelectedIndex = 0;
@@ -844,19 +929,21 @@ namespace UnrealVS
 				//RecentFilesListBox.SelectedIndex = 0;
 				AsyncFocusSelectedItem();
 
-				timer.Stop();
+				Timer.Stop();
 				RefreshingDialog.Visibility = Visibility.Collapsed;
-				IsRefreshing = false;
+				bIsRefreshing = false;
 			};
-			timer.Start();
+			Timer.Start();
 		}
 
-		private void FileBrowserWindowControl_IsVisibleChanged(object sender, DependencyPropertyChangedEventArgs e)
+		private void FileBrowserWindowControl_IsVisibleChanged(object Sender, DependencyPropertyChangedEventArgs E)
 		{
-			if (!(bool)e.NewValue)
+			if (!(bool)E.NewValue)
+			{
 				return;
+			}
 			FilterEditBox.SelectAll();
-			ClearFilter = true;
+			bClearFilter = true;
 			//FilesListTab.SelectedIndex = 0;
 
 			if (AllFileItems.Length == 0)
@@ -885,87 +972,110 @@ namespace UnrealVS
 		private void SaveSolutionSettings()
 		{
 			if (String.IsNullOrEmpty(UnrealVsFileName))
+			{
 				return;
-			using (var file = File.CreateText(UnrealVsFileName))
+			}
+			using (StreamWriter Writer = File.CreateText(UnrealVsFileName))
 			{
-				var settings = new FileBrowserSettings();
-				foreach (var item in BookmarkedFileItems)
-					settings.Bookmarks.Add(item.File);
-				foreach (var item in RecentFileItems)
-					settings.Recents.Add(item.File);
-				string json = JsonSerializer.Serialize(settings, new JsonSerializerOptions { WriteIndented = true });
-				file.Write(json);
+				FileBrowserSettings Settings = new FileBrowserSettings();
+				foreach (FileItem Item in BookmarkedFileItems)
+					Settings.Bookmarks.Add(Item.File);
+				foreach (FileItem Item in RecentFileItems)
+					Settings.Recents.Add(Item.File);
+				string Json = JsonSerializer.Serialize(Settings, new JsonSerializerOptions { WriteIndented = true });
+				Writer.Write(Json);
 			}
 		}
 
 		private void LoadSolutionSettings()
 		{
-			var solutionFileName = UnrealVSPackage.Instance.DTE.Solution.FileName;
-			if (String.IsNullOrEmpty(solutionFileName))
+			string SolutionFileName = UnrealVSPackage.Instance.DTE.Solution.FileName;
+			if (String.IsNullOrEmpty(SolutionFileName))
 			{
 				BookmarkedFileItems.Clear();
 				RecentFileItems.Clear();
 				UnrealVsFileName = null;
 				return;
 			}
-			UnrealVsFileName = solutionFileName.Substring(0, solutionFileName.Length - 3) + "unrealvs";
+			UnrealVsFileName = SolutionFileName.Substring(0, SolutionFileName.Length - 3) + "unrealvs";
 			if (!File.Exists(UnrealVsFileName))
+			{
 				return;
-			string json = File.ReadAllText(UnrealVsFileName);
-			if (String.IsNullOrEmpty(json))
+			}
+			string Json = File.ReadAllText(UnrealVsFileName);
+			if (String.IsNullOrEmpty(Json))
+			{
 				return;
-			var settings = JsonSerializer.Deserialize<FileBrowserSettings>(json, new JsonSerializerOptions { PropertyNameCaseInsensitive = true });
-			if (settings == null)
+			}
+			FileBrowserSettings Settings = JsonSerializer.Deserialize<FileBrowserSettings>(Json, new JsonSerializerOptions { PropertyNameCaseInsensitive = true });
+			if (Settings == null)
+			{
 				return;
+			}
 
 
 			BookmarkedFileItems.Clear();
-			if (settings.Bookmarks.Count == 0 && settings.Recents.Count == 0)
+			if (Settings.Bookmarks.Count == 0 && Settings.Recents.Count == 0)
 				return;
 
-			var fileItems = new[] { BookmarkedFileItems, RecentFileItems };
-			var settingItems = new[] { settings.Bookmarks, settings.Recents };
-			FileItem fileItem;
+			List<FileItem>[] FileItemLists = new[] { BookmarkedFileItems, RecentFileItems };
+			List<string>[] SettingItems = new[] { Settings.Bookmarks, Settings.Recents };
+			FileItem FoundItem;
 
 			for (int i = 0; i != 2; ++i)
 			{
-				var added = new HashSet<FileItem>();
-				foreach (var item in settingItems[i])
-					if (AllFileItemsLookup.TryGetValue(item, out fileItem))
-						if (added.Add(fileItem))
-							fileItems[i].Add(fileItem);
+				HashSet<FileItem> Added = new HashSet<FileItem>();
+				foreach (string Item in SettingItems[i])
+				{
+					if (AllFileItemsLookup.TryGetValue(Item, out FoundItem))
+					{
+						if (Added.Add(FoundItem))
+						{
+							FileItemLists[i].Add(FoundItem);
+						}
+					}
+				}
 			}
 		}
 
-		private static ListBoxItem GetVisibleListViewElement(ListBox listView, int y)
+		private static ListBoxItem GetVisibleListViewElement(ListBox ListView, int Y)
 		{
-			HitTestResult hitTest = VisualTreeHelper.HitTest(listView, new Point(10, y));
-			DependencyObject depObj = hitTest.VisualHit as DependencyObject;
-			if (depObj == null)
+			HitTestResult HitTest = VisualTreeHelper.HitTest(ListView, new Point(10, Y));
+			if (HitTest == null)
+			{
 				return null;
-			DependencyObject current = depObj;
-			while (current != null && current != listView)
+			}
+			DependencyObject DepObj = HitTest.VisualHit as DependencyObject;
+			if (DepObj == null)
+			{
+				return null;
+			}
+			DependencyObject Current = DepObj;
+			while (Current != null && Current != ListView)
 			{
-				ListBoxItem listBoxItem = current as ListBoxItem;
-				if (listBoxItem != null)
-					return listBoxItem;
-				current = VisualTreeHelper.GetParent(current);
+				ListBoxItem BoxItem = Current as ListBoxItem;
+				if (BoxItem != null)
+				{
+					return BoxItem;
+				}
+				Current = VisualTreeHelper.GetParent(Current);
 			}
 			return null;
 		}
-		private static char GetCharFromKey(Key key)
+
+		private static char GetCharFromKey(Key Key)
 		{
-			char ch = '\0';
+			char Ch = '\0';
 
-			int virtualKey = KeyInterop.VirtualKeyFromKey(key);
-			byte[] keyboardState = new byte[256];
-			NativeMethods.GetKeyboardState(keyboardState);
+			int VirtualKey = KeyInterop.VirtualKeyFromKey(Key);
+			byte[] KeyboardState = new byte[256];
+			NativeMethods.GetKeyboardState(KeyboardState);
 
-			uint scanCode = NativeMethods.MapVirtualKey((uint)virtualKey, NativeMethods.MapType.MAPVK_VK_TO_VSC);
-			StringBuilder stringBuilder = new StringBuilder(2);
+			uint ScanCode = NativeMethods.MapVirtualKey((uint)VirtualKey, NativeMethods.MapType.MAPVK_VK_TO_VSC);
+			StringBuilder Sb = new StringBuilder(2);
 
-			int result = NativeMethods.ToUnicode((uint)virtualKey, scanCode, keyboardState, stringBuilder, stringBuilder.Capacity, 0);
-			switch (result)
+			int Result = NativeMethods.ToUnicode((uint)VirtualKey, ScanCode, KeyboardState, Sb, Sb.Capacity, 0);
+			switch (Result)
 			{
 				case -1:
 					break;
@@ -973,20 +1083,20 @@ namespace UnrealVS
 					break;
 				case 1:
 				default:
-					ch = stringBuilder[0];
+					Ch = Sb[0];
 					break;
 			}
-			return ch;
+			return Ch;
 		}
 
-		class StackItem
+		private class StackItem
 		{
 			public ProjectItems Items;
 			public int Index;
 			public Project Project;
 		}
 
-		class SolutionTraverser
+		private class SolutionTraverser
 		{
 			private SortedDictionary<string, FileItem> SortedFileItems = new SortedDictionary<string, FileItem>();
 			private Dictionary<string, SortedDictionary<string, FileItem>> CollidingFileItems = new Dictionary<string, SortedDictionary<string, FileItem>>();
@@ -998,158 +1108,170 @@ namespace UnrealVS
 
 			public FileItem[] Update()
 			{
-				Stack<StackItem> itemStack = ItemStack;
-				var projects = UnrealVSPackage.Instance.DTE.Solution.Projects;
-				var projectCount = projects.Count;
+				Stack<StackItem> TraversalStack = ItemStack;
+				Projects Projects = UnrealVSPackage.Instance.DTE.Solution.Projects;
+				int ProjectCount = Projects.Count;
 
-				int traverseCounter = 0;
+				int TraverseCounter = 0;
 
-				while (ProjectIndex < projectCount)
+				while (ProjectIndex < ProjectCount)
 				{
-					var project = projects.Item(ProjectIndex + 1);
+					Project CurrentProject = Projects.Item(ProjectIndex + 1);
 
-					StackItem stackItem;
+					StackItem CurrentItem;
 					if (ItemStack.Count > 0)
-						stackItem = ItemStack.Pop();
+					{
+						CurrentItem = ItemStack.Pop();
+					}
 					else
-						stackItem = new StackItem() { Items = project.ProjectItems, Project = project };
+					{
+						CurrentItem = new StackItem() { Items = CurrentProject.ProjectItems, Project = CurrentProject };
+					}
 
 					while (true)
 					{
-						if (traverseCounter > 5000)
+						if (TraverseCounter > 5000)
 						{
-							ItemStack.Push(stackItem);
+							ItemStack.Push(CurrentItem);
 							return null;
 						}
 
-						if (stackItem.Items == null || stackItem.Index == stackItem.Items.Count)
+						if (CurrentItem.Items == null || CurrentItem.Index == CurrentItem.Items.Count)
 						{
-							if (itemStack.Count == 0)
+							if (TraversalStack.Count == 0)
+							{
 								break;
-							stackItem = itemStack.Pop();
-							++stackItem.Index;
+							}
+							CurrentItem = TraversalStack.Pop();
+							++CurrentItem.Index;
 							continue;
 						}
 
-						var projectItem = stackItem.Items.Item(stackItem.Index + 1);
-						++traverseCounter;
+						ProjectItem ProjectItem = CurrentItem.Items.Item(CurrentItem.Index + 1);
+						++TraverseCounter;
 
-						if (projectItem.FileCount != 0)
+						if (ProjectItem.FileCount != 0)
 						{
-							var file = projectItem.FileNames[1];
-							if (file != null)
+							string FileName = ProjectItem.FileNames[1];
+							if (FileName != null)
 							{
-								var name = projectItem.Name;
-								if (name != file)
-									if (!file.EndsWith("\\")) // Skip folders
+								string Name = ProjectItem.Name;
+								if (Name != FileName)
+								{
+									if (!FileName.EndsWith("\\")) // Skip folders
 									{
-										FileItem fileItem = new FileItem() { Name = name, File = file, Project = stackItem.Project.Name };
-										FileItem existingFileItem;
-										if (SortedFileItems.TryGetValue(name, out existingFileItem))
+										FileItem NewFileItem = new FileItem() { Name = Name, File = FileName, Project = CurrentItem.Project.Name };
+										FileItem ExistingFileItem;
+										if (SortedFileItems.TryGetValue(Name, out ExistingFileItem))
 										{
-											if (existingFileItem.File != file)
+											if (ExistingFileItem.File != FileName)
 											{
-												SortedDictionary<string, FileItem> colList;
-												if (!CollidingFileItems.TryGetValue(name, out colList))
+												SortedDictionary<string, FileItem> ColList;
+												if (!CollidingFileItems.TryGetValue(Name, out ColList))
+												{
+													ColList = new SortedDictionary<string, FileItem>();
+													ColList.Add(ExistingFileItem.File, ExistingFileItem);
+													CollidingFileItems[Name] = ColList;
+												}
+												if (!ColList.ContainsKey(NewFileItem.File))
 												{
-													colList = new SortedDictionary<string, FileItem>();
-													colList.Add(existingFileItem.File, existingFileItem);
-													CollidingFileItems[name] = colList;
+													ColList.Add(NewFileItem.File, NewFileItem);
 												}
-												if (!colList.ContainsKey(fileItem.File))
-													colList.Add(fileItem.File, fileItem);
 											}
 										}
 										else
-											SortedFileItems.Add(name, fileItem);
+											SortedFileItems.Add(Name, NewFileItem);
 
-										if (file.EndsWith(".build.cs", StringComparison.OrdinalIgnoreCase))
+										if (FileName.EndsWith(".build.cs", StringComparison.OrdinalIgnoreCase))
 										{
-											var lastBackslash = file.LastIndexOf('\\');
-											var moduleDir = file.Substring(0, lastBackslash);
-											var moduleName = file.Substring(lastBackslash + 1, file.Length - lastBackslash - 10);
-											if (!ModuleRoots.ContainsKey(moduleDir))
-												ModuleRoots.Add(moduleDir, moduleName);
+											int LastBackslash = FileName.LastIndexOf('\\');
+											string ModuleDir = FileName.Substring(0, LastBackslash);
+											string ModuleName = FileName.Substring(LastBackslash + 1, FileName.Length - LastBackslash - 10);
+											if (!ModuleRoots.ContainsKey(ModuleDir))
+											{
+												ModuleRoots.Add(ModuleDir, ModuleName);
+											}
 											//else if (ModuleRoots[moduleDir] != moduleName)
 											//    Logging.WriteLine(file);
 										}
 
 										++HandledItemCount;
 									}
+								}
 							}
 						}
 
-						if (projectItem.SubProject != null)
+						if (ProjectItem.SubProject != null)
 						{
-							itemStack.Push(stackItem);
-							stackItem = new StackItem() { Items = projectItem.SubProject.ProjectItems, Project = projectItem.SubProject };
+							TraversalStack.Push(CurrentItem);
+							CurrentItem = new StackItem() { Items = ProjectItem.SubProject.ProjectItems, Project = ProjectItem.SubProject };
 							continue;
 						}
 
-						if (projectItem.ProjectItems != null)
+						if (ProjectItem.ProjectItems != null)
 						{
-							itemStack.Push(stackItem);
-							stackItem = new StackItem() { Items = projectItem.ProjectItems, Project = stackItem.Project };
+							TraversalStack.Push(CurrentItem);
+							CurrentItem = new StackItem() { Items = ProjectItem.ProjectItems, Project = CurrentItem.Project };
 							continue;
 						}
 
-						++stackItem.Index;
+						++CurrentItem.Index;
 					}
 
 					++ProjectIndex;
 				}
 
 				// We need to rename all collisions
-				foreach (var kv in CollidingFileItems)
+				foreach (KeyValuePair<string, SortedDictionary<string, FileItem>> CollisionKv in CollidingFileItems)
 				{
-					SortedFileItems.Remove(kv.Key);
-					var usedParents = new Dictionary<string, int>();
+					SortedFileItems.Remove(CollisionKv.Key);
+					Dictionary<string, int> UsedParents = new Dictionary<string, int>();
 
-					foreach (var itemKv in kv.Value)
+					foreach (KeyValuePair<string, FileItem> ItemKv in CollisionKv.Value)
 					{
-						var item = itemKv.Value;
-						string parent = item.Project;
-						var path = item.File;
-						int lastBackslash = path.LastIndexOf('\\');
-						while (lastBackslash != -1)
+						FileItem Item = ItemKv.Value;
+						string Parent = Item.Project;
+						string FilePath = Item.File;
+						int LastBackslash = FilePath.LastIndexOf('\\');
+						while (LastBackslash != -1)
 						{
-							path = path.Substring(0, lastBackslash);
-							string moduleName;
-							if (ModuleRoots.TryGetValue(path, out moduleName))
+							FilePath = FilePath.Substring(0, LastBackslash);
+							string ModuleName;
+							if (ModuleRoots.TryGetValue(FilePath, out ModuleName))
 							{
-								parent = moduleName;
+								Parent = ModuleName;
 								break;
 							}
-							lastBackslash = path.LastIndexOf('\\');
+							LastBackslash = FilePath.LastIndexOf('\\');
 						}
 
-						string sortName = item.Name + $" ({parent})";
+						string SortName = Item.Name + $" ({Parent})";
 
-						int counter = 0;
-						if (usedParents.TryGetValue(parent, out counter))
+						int Counter = 0;
+						if (UsedParents.TryGetValue(Parent, out Counter))
 						{
-							if (counter == 0)
+							if (Counter == 0)
 							{
-								var alreadyAddedItem = SortedFileItems[sortName];
-								SortedFileItems.Remove(sortName);
-								alreadyAddedItem.Name = item.Name + $" ({parent} 0)";
-								SortedFileItems.Add(alreadyAddedItem.Name, alreadyAddedItem);
+								FileItem AlreadyAddedItem = SortedFileItems[SortName];
+								SortedFileItems.Remove(SortName);
+								AlreadyAddedItem.Name = Item.Name + $" ({Parent} 0)";
+								SortedFileItems.Add(AlreadyAddedItem.Name, AlreadyAddedItem);
 							}
-							++counter;
-							usedParents[parent] = counter;
-							item.Name = item.Name + $" ({parent} {counter})";
-							SortedFileItems.Add(item.Name, item);
+							++Counter;
+							UsedParents[Parent] = Counter;
+							Item.Name = Item.Name + $" ({Parent} {Counter})";
+							SortedFileItems.Add(Item.Name, Item);
 						}
 						else
 						{
-							usedParents.Add(parent, 0);
-							item.Name = sortName;
-							SortedFileItems.Add(item.Name, item);
+							UsedParents.Add(Parent, 0);
+							Item.Name = SortName;
+							SortedFileItems.Add(Item.Name, Item);
 						}
 					}
 				}
 
-				var result = SortedFileItems.Values.ToArray();
+				FileItem[] Result = SortedFileItems.Values.ToArray();
 
 				// Sort so header is before source file.. since most of the time you are opening header files.
 				/*
@@ -1171,8 +1293,8 @@ namespace UnrealVS
                     }
                 }
                 */
-				return result;
+				return Result;
 			}
 		}
 	}
-}
\ No newline at end of file
+}
