[UnrealVS] Add option to not close the File Browser window after a file is select.
Add the module and file path to the File Browser UI.
Change File Browser style to see the cursor.
Refactor Ctrl+i to use CopyIncludePath.

#Synced-CL 52420744
#Synced-CL 52420744

diff --git a/Engine/Extras/UnrealVS/VS2022/AtomFeed.xml b/Engine/Extras/UnrealVS/VS2022/AtomFeed.xml
--- a/Engine/Extras/UnrealVS/VS2022/AtomFeed.xml
+++ b/Engine/Extras/UnrealVS/VS2022/AtomFeed.xml
@@ -2,13 +2,13 @@
 <feed xmlns="http://www.w3.org/2005/Atom">
   <title type="text">UnrealVS Gallery</title>
   <id>uuid:45b4571a-2981-49f1-9f70-04cfa96db5c8;id=1</id>
-  <updated>2025-05-02T11:21:31Z</updated>
+  <updated>2026-04-07T00:00:00Z</updated>
   <entry>
     <id>ddbf523f-7eb6-4887-bd51-85a714ff87eb</id>
     <title type="text">UnrealVS</title>
     <summary type="text">Unreal Engine extension for Visual Studio</summary>
     <published>2022-10-21T10:01:57-07:00</published>
-    <updated>2025-05-08T11:21:31+00:00</updated>
+    <updated>2026-04-07T00:00:00+00:00</updated>
     <author>
       <name>Epic Games, Inc.</name>
     </author>
@@ -19,7 +19,7 @@
     <content type="application/octet-stream" src="UnrealVS.vsix" />
     <Vsix xmlns:xsd="http://www.w3.org/2001/XMLSchema" xmlns:xsi="http://www.w3.org/2001/XMLSchema-instance" xmlns="http://schemas.microsoft.com/developer/vsx-syndication-schema/2010">
       <Id>ddbf523f-7eb6-4887-bd51-85a714ff87eb</Id>
-      <Version>1.84.9.6</Version>
+      <Version>1.85.0.0</Version>
       <Rating xsi:nil="true" />
       <RatingCount xsi:nil="true" />
       <DownloadCount xsi:nil="true" />
diff --git a/Engine/Source/Programs/UnrealVS/ReleaseNotes.txt b/Engine/Source/Programs/UnrealVS/ReleaseNotes.txt
--- a/Engine/Source/Programs/UnrealVS/ReleaseNotes.txt
+++ b/Engine/Source/Programs/UnrealVS/ReleaseNotes.txt
@@ -1,3 +1,8 @@
+1.85.0.0
+* Add option to not close the File Browser window after a file is select.
+* Add the module and file path to File Browser.
+* Add s context menu to File Browser.
+
 1.84.9.6
 * Usability improvements for startup project selector
 
diff --git a/Engine/Source/Programs/UnrealVS/UnrealVS.2022/Properties/AssemblyInfo.cs b/Engine/Source/Programs/UnrealVS/UnrealVS.2022/Properties/AssemblyInfo.cs
--- a/Engine/Source/Programs/UnrealVS/UnrealVS.2022/Properties/AssemblyInfo.cs
+++ b/Engine/Source/Programs/UnrealVS/UnrealVS.2022/Properties/AssemblyInfo.cs
@@ -31,6 +31,6 @@ using System.Runtime.InteropServices;
 // You can specify all the values or you can default the Revision and Build Numbers 
 // by using the '*' as shown below:
 
-[assembly: AssemblyVersion("1.84.9.6")]
-[assembly: AssemblyFileVersion("1.84.9.6")]
+[assembly: AssemblyVersion("1.85.0.0")]
+[assembly: AssemblyFileVersion("1.85.0.0")]
 
diff --git a/Engine/Source/Programs/UnrealVS/UnrealVS.2022/UnrealVS.2022.csproj b/Engine/Source/Programs/UnrealVS/UnrealVS.2022/UnrealVS.2022.csproj
--- a/Engine/Source/Programs/UnrealVS/UnrealVS.2022/UnrealVS.2022.csproj
+++ b/Engine/Source/Programs/UnrealVS/UnrealVS.2022/UnrealVS.2022.csproj
@@ -112,6 +112,16 @@
     </Resource>
     <Resource Include="Resources\RefreshProjects.png">
     </Resource>
+    <Resource Include="Resources\FileBrowser.png">
+    </Resource>
+    <Resource Include="Resources\BookmarkAdd.png">
+    </Resource>
+    <Resource Include="Resources\BookmarkRemove.png">
+    </Resource>
+    <Resource Include="Resources\CopyInclude.png">
+    </Resource>
+    <Resource Include="Resources\InsertInclude.png">
+    </Resource>
     <Resource Include="Resources\UnrealIcon.png">
     </Resource>
     <None Include="source.extension.vsixmanifest" />
diff --git a/Engine/Source/Programs/UnrealVS/UnrealVS.Shared/FileBrowser.cs b/Engine/Source/Programs/UnrealVS/UnrealVS.Shared/FileBrowser.cs
--- a/Engine/Source/Programs/UnrealVS/UnrealVS.Shared/FileBrowser.cs
+++ b/Engine/Source/Programs/UnrealVS/UnrealVS.Shared/FileBrowser.cs
@@ -41,6 +41,10 @@ namespace UnrealVS
 			}
 			IVsWindowFrame ToolWindowFrame = (IVsWindowFrame)ToolWindow.Frame;
 			Microsoft.VisualStudio.ErrorHandler.ThrowOnFailure(ToolWindowFrame.Show());
+
+			// IsVisibleChanged only fires when visibility actually changes; if the window
+			// was already visible, focus the filter box explicitly every time.
+			BrowserControl?.HandleShow();
 		}
 
 		FileBrowserWindowControl BrowserControl
diff --git a/Engine/Source/Programs/UnrealVS/UnrealVS.Shared/FileBrowserWindowControl.xaml b/Engine/Source/Programs/UnrealVS/UnrealVS.Shared/FileBrowserWindowControl.xaml
--- a/Engine/Source/Programs/UnrealVS/UnrealVS.Shared/FileBrowserWindowControl.xaml
+++ b/Engine/Source/Programs/UnrealVS/UnrealVS.Shared/FileBrowserWindowControl.xaml
@@ -22,6 +22,14 @@
             <Setter Property="Margin" Value="0,-6,0,0"/>
             <Setter Property="FocusVisualStyle" Value="{x:Null}"/>
         </Style>
+        <Style TargetType="ListView">
+            <Setter Property="Background" Value="Transparent"/>
+            <Setter Property="BorderThickness" Value="1"/>
+            <Setter Property="BorderBrush" Value="{DynamicResource {x:Static vsshell:VsBrushes.ToolWindowBorderKey}}"/>
+            <Setter Property="Margin" Value="0,-6,0,0"/>
+            <Setter Property="FocusVisualStyle" Value="{x:Null}"/>
+            <Setter Property="ScrollViewer.HorizontalScrollBarVisibility" Value="Auto"/>
+        </Style>
         <Style TargetType="TabItem">
             <Setter Property="Background" Value="Transparent"/>
             <Setter Property="IsTabStop" Value="False"/>
@@ -86,10 +94,92 @@
                 </MultiDataTrigger>
             </Style.Triggers>
         </Style>
+        <Style TargetType="ListViewItem">
+            <Setter Property="Foreground" Value="{DynamicResource {x:Static vsshell:VsBrushes.CommandBarTextHoverKey}}"/>
+            <Setter Property="Template">
+                <Setter.Value>
+                    <ControlTemplate TargetType="{x:Type ListViewItem}">
+                        <Border BorderBrush="Transparent" BorderThickness="0" Background="{TemplateBinding Background}">
+                            <GridViewRowPresenter Margin="2,1"/>
+                        </Border>
+                    </ControlTemplate>
+                </Setter.Value>
+            </Setter>
+            <Style.Triggers>
+                <MultiTrigger>
+                    <MultiTrigger.Conditions>
+                        <Condition Property="IsSelected" Value="true" />
+                    </MultiTrigger.Conditions>
+                    <Setter Property="Background" Value="{DynamicResource {x:Static vsshell:VsBrushes.ToolWindowTabMouseOverBackgroundBeginKey}}" />
+                    <Setter Property="FocusVisualStyle" Value="{StaticResource ListBoxItemVisualStyle}" />
+                </MultiTrigger>
+                <MultiDataTrigger>
+                    <MultiDataTrigger.Conditions>
+                        <Condition Binding="{Binding Path=Name}" Value="" />
+                    </MultiDataTrigger.Conditions>
+                    <Setter Property="Background" Value="{x:Null}" />
+                </MultiDataTrigger>
+            </Style.Triggers>
+        </Style>
+        <Style TargetType="{x:Type GridViewColumnHeader}">
+            <Setter Property="Background" Value="{DynamicResource {x:Static vsshell:VsBrushes.ToolWindowBackgroundKey}}"/>
+            <Setter Property="Foreground" Value="{DynamicResource {x:Static vsshell:VsBrushes.WindowTextKey}}"/>
+            <Setter Property="HorizontalContentAlignment" Value="Left"/>
+            <Setter Property="Padding" Value="5,3,5,3"/>
+            <Setter Property="Template">
+                <Setter.Value>
+                    <ControlTemplate TargetType="{x:Type GridViewColumnHeader}">
+                        <Grid SnapsToDevicePixels="True">
+                            <Border x:Name="HeaderBorder"
+                                    BorderThickness="0,0,0,1"
+                                    BorderBrush="{DynamicResource {x:Static vsshell:VsBrushes.ToolWindowBorderKey}}"
+                                    Background="{TemplateBinding Background}">
+                                <ContentPresenter Margin="{TemplateBinding Padding}"
+                                                  VerticalAlignment="Center"
+                                                  HorizontalAlignment="{TemplateBinding HorizontalContentAlignment}"
+                                                  RecognizesAccessKey="True"/>
+                            </Border>
+                            <Thumb x:Name="PART_HeaderGripper" HorizontalAlignment="Right" Width="8" Cursor="SizeWE">
+                                <Thumb.Style>
+                                    <Style TargetType="{x:Type Thumb}">
+                                        <Setter Property="Template">
+                                            <Setter.Value>
+                                                <ControlTemplate TargetType="{x:Type Thumb}">
+                                                    <Border Background="Transparent"
+                                                            BorderBrush="{DynamicResource {x:Static vsshell:VsBrushes.ToolWindowBorderKey}}"
+                                                            BorderThickness="0,0,1,0"/>
+                                                </ControlTemplate>
+                                            </Setter.Value>
+                                        </Setter>
+                                    </Style>
+                                </Thumb.Style>
+                            </Thumb>
+                        </Grid>
+                        <ControlTemplate.Triggers>
+                            <Trigger Property="IsMouseOver" Value="True">
+                                <Setter TargetName="HeaderBorder" Property="Background"
+                                        Value="{DynamicResource {x:Static vsshell:VsBrushes.ToolWindowTabMouseOverBackgroundBeginKey}}"/>
+                            </Trigger>
+                            <Trigger Property="IsPressed" Value="True">
+                                <Setter TargetName="HeaderBorder" Property="Background"
+                                        Value="{DynamicResource {x:Static vsshell:VsBrushes.CommandBarSelectedBorderKey}}"/>
+                            </Trigger>
+                            <Trigger Property="Role" Value="Padding">
+                                <Setter TargetName="PART_HeaderGripper" Property="Visibility" Value="Collapsed"/>
+                            </Trigger>
+                        </ControlTemplate.Triggers>
+                    </ControlTemplate>
+                </Setter.Value>
+            </Setter>
+        </Style>
         <Style TargetType="Border">
             <Setter Property="BorderBrush" Value="{DynamicResource {x:Static vsshell:VsBrushes.ToolWindowBorderKey}}"/>
             <Setter Property="BorderThickness" Value="1"/>
         </Style>
+        <Image x:Key="IconBookmarkAdd"    x:Shared="False" Source="/UnrealVS;component/Resources/BookmarkAdd.png"    Width="16" Height="16"/>
+        <Image x:Key="IconBookmarkRemove" x:Shared="False" Source="/UnrealVS;component/Resources/BookmarkRemove.png" Width="16" Height="16"/>
+        <Image x:Key="IconCopyInclude"    x:Shared="False" Source="/UnrealVS;component/Resources/CopyInclude.png"    Width="16" Height="16"/>
+        <Image x:Key="IconInsertInclude"  x:Shared="False" Source="/UnrealVS;component/Resources/InsertInclude.png"  Width="16" Height="16"/>
     </UserControl.Resources>
     <Grid>
         <Grid.RowDefinitions>
@@ -103,41 +193,111 @@
             </Grid.RowDefinitions>
             <Border Margin="13,16,13,0" Padding="5" BorderBrush="{DynamicResource {x:Static vsshell:VsBrushes.ToolWindowBorderKey}}">
                 <Grid>
-                <Grid.ColumnDefinitions>
-                    <ColumnDefinition Width="auto"/>
-                    <ColumnDefinition Width="*"/>
-                </Grid.ColumnDefinitions>
+                    <Grid.ColumnDefinitions>
+                        <ColumnDefinition Width="auto"/>
+                        <ColumnDefinition Width="*"/>
+                        <ColumnDefinition Width="auto"/>
+                    </Grid.ColumnDefinitions>
                     <TextBlock Text="Filter:" Padding="3"/>
                     <TextBox Grid.Column="1" Name="FilterEditBox" CaretBrush="{DynamicResource {x:Static vsshell:VsBrushes.CommandBarTextActiveKey}}" Padding="3"/>
+                    <Button Grid.Column="2" Name="RefreshButton" Click="RefreshButton_Click"
+                            Width="22" Height="22" Margin="3,0,0,0" Padding="3"
+                            ToolTip="Refresh (F5)"
+                            Background="Transparent" BorderThickness="0"
+                            VerticalAlignment="Center">
+                        <Path Width="12" Height="12" Stretch="Uniform"
+                              Data="M17.65 6.35C16.2 4.9 14.21 4 12 4c-4.42 0-7.99 3.58-7.99 8s3.57 8 7.99 8c3.73 0 6.84-2.55 7.73-6h-2.08c-.82 2.33-3.04 4-5.65 4-3.31 0-6-2.69-6-6s2.69-6 6-6c1.66 0 3.14.69 4.22 1.78L13 11h7V4l-2.35 2.35z"
+                              Fill="{DynamicResource {x:Static vsshell:VsBrushes.WindowTextKey}}"/>
+                    </Button>
                 </Grid>
             </Border>
             <TabControl Name="FilesListTab" Grid.Row="1" IsTabStop="False" Background="Transparent" Margin="10" BorderBrush="Transparent">
                 <TabItem Header="All files" Name="AllFilesTab" Focusable="False">
-                    <ListBox Name="AllFilesListBox">
-                            <ListBox.ItemTemplate>
-                                <DataTemplate>
-                                    <TextBlock Text="{Binding Path=Name}" />
-                                </DataTemplate>
-                            </ListBox.ItemTemplate>
-                        </ListBox>
+                    <ListView Name="AllFilesListBox" GridViewColumnHeader.Click="ColumnHeader_Click" PreviewMouseRightButtonDown="FileListView_PreviewMouseRightButtonDown">
+                        <ListView.View>
+                            <GridView>
+                                <GridViewColumn Header="Name" Width="200" DisplayMemberBinding="{Binding Name}">
+                                    <GridViewColumn.HeaderContainerStyle>
+                                        <Style TargetType="GridViewColumnHeader" BasedOn="{StaticResource {x:Type GridViewColumnHeader}}">
+                                            <Setter Property="Tag" Value="Name"/>
+                                        </Style>
+                                    </GridViewColumn.HeaderContainerStyle>
+                                </GridViewColumn>
+                                <GridViewColumn Header="Module" Width="130" DisplayMemberBinding="{Binding Module}">
+                                    <GridViewColumn.HeaderContainerStyle>
+                                        <Style TargetType="GridViewColumnHeader" BasedOn="{StaticResource {x:Type GridViewColumnHeader}}">
+                                            <Setter Property="Tag" Value="Module"/>
+                                        </Style>
+                                    </GridViewColumn.HeaderContainerStyle>
+                                </GridViewColumn>
+                                <GridViewColumn Header="Path" Width="500" DisplayMemberBinding="{Binding File}">
+                                    <GridViewColumn.HeaderContainerStyle>
+                                        <Style TargetType="GridViewColumnHeader" BasedOn="{StaticResource {x:Type GridViewColumnHeader}}">
+                                            <Setter Property="Tag" Value="File"/>
+                                        </Style>
+                                    </GridViewColumn.HeaderContainerStyle>
+                                </GridViewColumn>
+                            </GridView>
+                        </ListView.View>
+                    </ListView>
                 </TabItem>
                 <TabItem Header="Bookmarked files" Name="BookmarkedFilesTab" Focusable="False">
-                    <ListBox Name="BookmarkedFilesListBox">
-                        <ListBox.ItemTemplate>
-                            <DataTemplate>
-                                <TextBlock Text="{Binding Path=Name}" />
-                            </DataTemplate>
-                        </ListBox.ItemTemplate>
-                    </ListBox>
+                    <ListView Name="BookmarkedFilesListBox" GridViewColumnHeader.Click="ColumnHeader_Click" PreviewMouseRightButtonDown="FileListView_PreviewMouseRightButtonDown">
+                        <ListView.View>
+                            <GridView>
+                                <GridViewColumn Header="Name" Width="200" DisplayMemberBinding="{Binding Name}">
+                                    <GridViewColumn.HeaderContainerStyle>
+                                        <Style TargetType="GridViewColumnHeader" BasedOn="{StaticResource {x:Type GridViewColumnHeader}}">
+                                            <Setter Property="Tag" Value="Name"/>
+                                        </Style>
+                                    </GridViewColumn.HeaderContainerStyle>
+                                </GridViewColumn>
+                                <GridViewColumn Header="Module" Width="130" DisplayMemberBinding="{Binding Module}">
+                                    <GridViewColumn.HeaderContainerStyle>
+                                        <Style TargetType="GridViewColumnHeader" BasedOn="{StaticResource {x:Type GridViewColumnHeader}}">
+                                            <Setter Property="Tag" Value="Module"/>
+                                        </Style>
+                                    </GridViewColumn.HeaderContainerStyle>
+                                </GridViewColumn>
+                                <GridViewColumn Header="Path" Width="500" DisplayMemberBinding="{Binding File}">
+                                    <GridViewColumn.HeaderContainerStyle>
+                                        <Style TargetType="GridViewColumnHeader" BasedOn="{StaticResource {x:Type GridViewColumnHeader}}">
+                                            <Setter Property="Tag" Value="File"/>
+                                        </Style>
+                                    </GridViewColumn.HeaderContainerStyle>
+                                </GridViewColumn>
+                            </GridView>
+                        </ListView.View>
+                    </ListView>
                 </TabItem>
                 <TabItem Header="Recent files" Name="RecentFilesTab" Focusable="False">
-                    <ListBox Name="RecentFilesListBox">
-                        <ListBox.ItemTemplate>
-                            <DataTemplate>
-                                <TextBlock Text="{Binding Path=Name}" />
-                            </DataTemplate>
-                        </ListBox.ItemTemplate>
-                    </ListBox>
+                    <ListView Name="RecentFilesListBox" GridViewColumnHeader.Click="ColumnHeader_Click" PreviewMouseRightButtonDown="FileListView_PreviewMouseRightButtonDown">
+                        <ListView.View>
+                            <GridView>
+                                <GridViewColumn Header="Name" Width="200" DisplayMemberBinding="{Binding Name}">
+                                    <GridViewColumn.HeaderContainerStyle>
+                                        <Style TargetType="GridViewColumnHeader" BasedOn="{StaticResource {x:Type GridViewColumnHeader}}">
+                                            <Setter Property="Tag" Value="Name"/>
+                                        </Style>
+                                    </GridViewColumn.HeaderContainerStyle>
+                                </GridViewColumn>
+                                <GridViewColumn Header="Module" Width="130" DisplayMemberBinding="{Binding Module}">
+                                    <GridViewColumn.HeaderContainerStyle>
+                                        <Style TargetType="GridViewColumnHeader" BasedOn="{StaticResource {x:Type GridViewColumnHeader}}">
+                                            <Setter Property="Tag" Value="Module"/>
+                                        </Style>
+                                    </GridViewColumn.HeaderContainerStyle>
+                                </GridViewColumn>
+                                <GridViewColumn Header="Path" Width="500" DisplayMemberBinding="{Binding File}">
+                                    <GridViewColumn.HeaderContainerStyle>
+                                        <Style TargetType="GridViewColumnHeader" BasedOn="{StaticResource {x:Type GridViewColumnHeader}}">
+                                            <Setter Property="Tag" Value="File"/>
+                                        </Style>
+                                    </GridViewColumn.HeaderContainerStyle>
+                                </GridViewColumn>
+                            </GridView>
+                        </ListView.View>
+                    </ListView>
                 </TabItem>
             </TabControl>
             <Border Grid.Row="2" Margin="13,0,13,16" Padding="5" BorderBrush="{DynamicResource {x:Static vsshell:VsBrushes.ToolWindowBorderKey}}">
@@ -181,7 +341,7 @@
                         <TextBlock Text="- Delete bookmarked file"/>
                         <TextBlock Text="- This help"/>
                         <TextBlock Text="- Refresh file lists with current visual studio solution"/>
-                        <TextBlock Text="- Create a #include &quot;file&quot; statement to clipboard and close file browser"/>
+                        <TextBlock Text="- Create a #include &quot;file&quot; statement to clipboard"/>
                         <TextBlock Text="- Create a #include &quot;file&quot; statement and paste it into active document"/>
                     </StackPanel>
                 </Grid>
diff --git a/Engine/Source/Programs/UnrealVS/UnrealVS.Shared/FileBrowserWindowControl.xaml.cs b/Engine/Source/Programs/UnrealVS/UnrealVS.Shared/FileBrowserWindowControl.xaml.cs
--- a/Engine/Source/Programs/UnrealVS/UnrealVS.Shared/FileBrowserWindowControl.xaml.cs
+++ b/Engine/Source/Programs/UnrealVS/UnrealVS.Shared/FileBrowserWindowControl.xaml.cs
@@ -9,12 +9,14 @@ using System.IO;
 using System.Linq;
 using System.Text;
 using System.Text.Json;
+using System.ComponentModel;
 using System.Windows;
 using System.Windows.Controls;
 using System.Windows.Data;
 using System.Windows.Input;
 using System.Windows.Media;
 using System.Windows.Threading;
+using System.Windows.Media.Imaging;
 
 namespace UnrealVS
 {
@@ -30,8 +32,7 @@ namespace UnrealVS
 		{
 			public string Name { get; set; }
 			public string File { get; set; }
-			public string Project { get; set; }
-			public bool Selectable { get; set; }
+			public string Module { get; set; }
 		}
 
 		private string UnrealVsFileName;
@@ -42,6 +43,8 @@ namespace UnrealVS
 
 		private Dictionary<string, FileItem> AllFileItemsLookup;
 
+		private Dictionary<ListBox, (GridViewColumnHeader Header, string OriginalContent, ListSortDirection Direction)> ColumnSortState;
+
 		public FileBrowserWindowControl(FileBrowserWindow InWindow)
 		{
 			Window = InWindow;
@@ -50,6 +53,7 @@ namespace UnrealVS
 			AllFileItems = Array.Empty<FileItem>();
 			BookmarkedFileItems = new List<FileItem>();
 			RecentFileItems = new List<FileItem>();
+			ColumnSortState = new Dictionary<ListBox, (GridViewColumnHeader, string, ListSortDirection)>();
 			IsVisibleChanged += FileBrowserWindowControl_IsVisibleChanged;
 			foreach (ListBox FileListBox in new[] { AllFilesListBox, BookmarkedFilesListBox, RecentFilesListBox })
 			{
@@ -145,6 +149,265 @@ namespace UnrealVS
 			AsyncFocusSelectedItem();
 		}
 
+		private void RefreshButton_Click(object Sender, RoutedEventArgs E)
+		{
+			AsyncRefreshFileList();
+		}
+
+		// ── Context menu ─────────────────────────────────────────────────────────
+
+		private void FileListView_PreviewMouseRightButtonDown(object Sender, MouseButtonEventArgs E)
+		{
+			// Walk up the visual tree to find the ListViewItem under the cursor.
+			DependencyObject Obj = E.OriginalSource as DependencyObject;
+			while (Obj != null && !(Obj is ListViewItem))
+			{
+				Obj = VisualTreeHelper.GetParent(Obj);
+			}
+			if (!(Obj is ListViewItem))
+			{
+				return;
+			}
+
+			ListBox ListBox = (ListBox)Sender;
+			FileItem Item = ((ListViewItem)Obj).DataContext as FileItem;
+			if (Item == null || Item.Name == "")
+			{
+				return;
+			}
+
+			ListBox.SelectedItem = Item;
+			ListBox.Focus();
+
+			bool bIsBookmarked = BookmarkedFileItems.Contains(Item);
+			string IncludePath = GetIncludePath(Item);
+			bool bHasActiveDoc = UnrealVSPackage.Instance.DTE?.ActiveDocument != null;
+
+			ContextMenu Menu = new ContextMenu();
+
+			// 1. Bookmark toggle
+			MenuItem BookmarkItem = new MenuItem();
+			if (bIsBookmarked)
+			{
+				BookmarkItem.Header = "Remove from Bookmarks";
+				BookmarkItem.Icon = FindResource("IconBookmarkRemove");
+				BookmarkItem.ToolTip = "Remove this file from the bookmarked list (Delete)";
+			}
+			else
+			{
+				BookmarkItem.Header = "Add to Bookmarks";
+				BookmarkItem.Icon = FindResource("IconBookmarkAdd");
+				BookmarkItem.ToolTip = "Add this file to the bookmarked list (Insert)";
+			}
+			BookmarkItem.Tag = Item;
+			BookmarkItem.Click += FileItemContextMenu_ToggleBookmark_Click;
+			Menu.Items.Add(BookmarkItem);
+
+			Menu.Items.Add(new Separator());
+
+			// 2. Copy #include to clipboard
+			MenuItem CopyItem = new MenuItem();
+			CopyItem.Header = "Copy #include to Clipboard";
+			CopyItem.Icon = FindResource("IconCopyInclude");
+			CopyItem.ToolTip = "Copy a #include \"...\" statement to the clipboard (Ctrl+I)";
+			CopyItem.IsEnabled = IncludePath != null;
+			CopyItem.Tag = Item;
+			CopyItem.Click += FileItemContextMenu_CopyInclude_Click;
+			Menu.Items.Add(CopyItem);
+
+			// 3. Insert #include into active document
+			MenuItem InsertItem = new MenuItem();
+			InsertItem.Header = "Insert #include into Active Document";
+			InsertItem.Icon = FindResource("IconInsertInclude");
+			InsertItem.ToolTip = "Insert a #include \"...\" statement into the active document at the correct position (Ctrl+Shift+I)";
+			InsertItem.IsEnabled = IncludePath != null && bHasActiveDoc;
+			InsertItem.Tag = Item;
+			InsertItem.Click += FileItemContextMenu_InsertInclude_Click;
+			Menu.Items.Add(InsertItem);
+
+			Menu.PlacementTarget = (UIElement)Sender;
+			Menu.IsOpen = true;
+			E.Handled = true;
+		}
+
+		private void FileItemContextMenu_ToggleBookmark_Click(object Sender, RoutedEventArgs E)
+		{
+			FileItem Item = (FileItem)((MenuItem)Sender).Tag;
+			if (BookmarkedFileItems.Contains(Item))
+			{
+				BookmarkedFileItems.Remove(Item);
+			}
+			else
+			{
+				BookmarkedFileItems.Add(Item);
+			}
+			RefreshListView(BookmarkedFilesListBox, BookmarkedFileItems.ToArray());
+			SaveSolutionSettings();
+		}
+
+		private void FileItemContextMenu_CopyInclude_Click(object Sender, RoutedEventArgs E)
+		{
+			FileItem Item = (FileItem)((MenuItem)Sender).Tag;
+			string IncludePath = GetIncludePath(Item);
+			if (IncludePath == null)
+			{
+				return;
+			}
+			Clipboard.SetText($"#include \"{IncludePath}\"");
+		}
+
+		private void FileItemContextMenu_InsertInclude_Click(object Sender, RoutedEventArgs E)
+		{
+			FileItem Item = (FileItem)((MenuItem)Sender).Tag;
+			string IncludePath = GetIncludePath(Item);
+			if (IncludePath == null)
+			{
+				return;
+			}
+			Clipboard.SetText($"#include \"{IncludePath}\"");
+			InsertIncludeIntoDocument(Item);
+		}
+
+		// ── InsertIncludeIntoDocument ──────
+
+		private void InsertIncludeIntoDocument(FileItem Item)
+		{
+			if (Item == null)
+			{
+				return;
+			}
+			Document ActiveDoc = UnrealVSPackage.Instance.DTE.ActiveDocument;
+			if (ActiveDoc == null)
+			{
+				return;
+			}
+			TextDocument Doc = (TextDocument)(ActiveDoc.Object("TextDocument"));
+			if (Doc == null)
+			{
+				return;
+			}
+
+			bool bAddSpace = false;
+			string IncludePath = GetIncludePath(Item);
+			string PasteString = $"#include \"{IncludePath}\"";
+
+			bool bSkipFirstInclude = ActiveDoc.FullName.EndsWith(".cpp");
+			int AddBeforeLine = -1;
+			int LastIncludeLine = -1;
+
+			bool bIsInComment = false;
+			EditPoint EditPoint = Doc.StartPoint.CreateEditPoint();
+			int LastLine = Doc.EndPoint.Line;
+			for (int LineIndex = 1; LineIndex < LastLine; ++LineIndex)
+			{
+				string Str = EditPoint.GetLines(LineIndex, LineIndex + 1).Trim();
+				if (Str.Length == 0)
+				{
+					continue;
+				}
+
+				if (bIsInComment)
+				{
+					if (Str.IndexOf("*/") != -1)
+					{
+						bIsInComment = false;
+					}
+					continue;
+				}
+
+				if (Str.StartsWith("//"))
+				{
+					continue;
+				}
+
+				if (Str.StartsWith("/*"))
+				{
+					bIsInComment = true;
+					continue;
+				}
+
+				if (Str.StartsWith("#include"))
+				{
+					if (Str == PasteString)
+					{
+						Doc.Selection.MoveTo(LineIndex, 0);
+						Doc.Selection.SelectLine();
+						EditPoint.TryToShow();
+						return;
+					}
+
+					bool bIsFirst = LastIncludeLine == -1;
+					int PrevLastIncludeLine = LastIncludeLine;
+					LastIncludeLine = LineIndex;
+					if (bSkipFirstInclude && bIsFirst)
+					{
+						continue;
+					}
+
+					ReadOnlySpan<char> Span = Str.AsSpan(9);
+					int FirstQuote = Span.IndexOf('"');
+					if (FirstQuote != -1)
+					{
+						Span = Span.Slice(FirstQuote + 1);
+						int SecondQuote = Span.IndexOf('"');
+						if (SecondQuote != -1)
+						{
+							Span = Span.Slice(0, SecondQuote).Trim();
+							bool bIsGeneratedInclude = Span.IndexOf(".generated.".AsSpan()) != -1;
+							if (IncludePath.AsSpan().CompareTo(Span, StringComparison.Ordinal) < 0 || bIsGeneratedInclude)
+							{
+								if (bIsGeneratedInclude && PrevLastIncludeLine != -1)
+								{
+									LastIncludeLine = PrevLastIncludeLine;
+								}
+								else
+								{
+									AddBeforeLine = LineIndex;
+								}
+								break;
+							}
+						}
+					}
+					else if (Span.IndexOf("UE_INLINE_GENERATED_CPP".AsSpan()) != -1)
+					{
+						LastIncludeLine = PrevLastIncludeLine;
+						break;
+					}
+					continue;
+				}
+
+				if (Str.StartsWith("#pragma once"))
+				{
+					LastIncludeLine = LineIndex + 1;
+					continue;
+				}
+
+				LastIncludeLine = LineIndex - 1;
+				bAddSpace = true;
+				break;  // Unknown code
+			}
+
+			if (AddBeforeLine == -1 && LastIncludeLine != -1)
+			{
+				AddBeforeLine = LastIncludeLine + 1;
+			}
+
+			if (AddBeforeLine == -1)
+			{
+				AddBeforeLine = LastLine;
+			}
+
+			EditPoint.LineDown(AddBeforeLine - 1);
+			EditPoint.Insert($"{PasteString}\r\n");
+			if (bAddSpace)
+			{
+				EditPoint.Insert("\r\n");
+			}
+			Doc.Selection.MoveTo(AddBeforeLine, 0);
+			Doc.Selection.SelectLine();
+			EditPoint.TryToShow();
+		}
+
 		private void FilterEditBox_KeyDown(object Sender, KeyEventArgs E)
 		{
 			if (E.Key == Key.Tab)
@@ -198,7 +461,9 @@ namespace UnrealVS
 			if (E.Key == Key.Enter)
 			{
 				if (OpenSelectedFile())
-					HidePanel();
+				{
+					HidePanelAfterSelection();
+				}
 				return;
 			}
 			if (E.Key == Key.PageDown)
@@ -280,133 +545,7 @@ namespace UnrealVS
 						return;
 					}
 
-					// Let's add directly in to active document
-					Document ActiveDoc = UnrealVSPackage.Instance.DTE.ActiveDocument;
-					TextDocument Doc = (TextDocument)(ActiveDoc.Object("TextDocument"));
-					if (Doc == null)
-					{
-						return;
-					}
-
-					bool bAddSpace = false;
-					string IncludePath = GetIncludePath((FileItem)(ActiveFilesListBox.SelectedItems[0]));
-					string PasteString = $"#include \"{IncludePath}\"";
-
-					bool bSkipFirstInclude = ActiveDoc.FullName.EndsWith(".cpp");
-					int AddBeforeLine = -1;
-					int LastIncludeLine = -1;
-
-					bool bIsInComment = false;
-					EditPoint EditPoint = Doc.StartPoint.CreateEditPoint();
-					int LastLine = Doc.EndPoint.Line;
-					for (int LineIndex = 1; LineIndex < LastLine; ++LineIndex)
-					{
-						string Str = EditPoint.GetLines(LineIndex, LineIndex + 1).Trim();
-						if (Str.Length == 0)
-						{
-							continue;
-						}
-
-						if (bIsInComment)
-						{
-							if (Str.IndexOf("*/") != -1)
-							{
-								bIsInComment = false;
-							}
-							continue;
-						}
-
-						if (Str.StartsWith("//"))
-						{
-							continue;
-						}
-
-						if (Str.StartsWith("/*"))
-						{
-							bIsInComment = true;
-							continue;
-						}
-
-						if (Str.StartsWith("#include"))
-						{
-							if (Str == PasteString)
-							{
-								Doc.Selection.MoveTo(LineIndex, 0);
-								Doc.Selection.SelectLine();
-								EditPoint.TryToShow();
-								return;
-							}
-
-							bool bIsFirst = LastIncludeLine == -1;
-							int PrevLastIncludeLine = LastIncludeLine;
-							LastIncludeLine = LineIndex;
-							if (bSkipFirstInclude && bIsFirst)
-							{
-								continue;
-							}
-
-							ReadOnlySpan<char> Span = Str.AsSpan(9);
-							int FirstQuote = Span.IndexOf('"');
-							if (FirstQuote != -1)
-							{
-								Span = Span.Slice(FirstQuote + 1);
-								int SecondQuote = Span.IndexOf('"');
-								if (SecondQuote != -1)
-								{
-									Span = Span.Slice(0, SecondQuote).Trim();
-									bool bIsGeneratedInclude = Span.IndexOf(".generated.".AsSpan()) != -1;
-									if (IncludePath.AsSpan().CompareTo(Span, StringComparison.Ordinal) < 0 || bIsGeneratedInclude)
-									{
-										if (bIsGeneratedInclude && PrevLastIncludeLine != -1)
-										{
-											LastIncludeLine = PrevLastIncludeLine;
-										}
-										else
-										{
-											AddBeforeLine = LineIndex;
-										}
-										break;
-									}
-								}
-							}
-							else if (Span.IndexOf("UE_INLINE_GENERATED_CPP".AsSpan()) != -1)
-							{
-								LastIncludeLine = PrevLastIncludeLine;
-								break;
-							}
-							continue;
-						}
-
-						if (Str.StartsWith("#pragma once"))
-						{
-							LastIncludeLine = LineIndex + 1;
-							continue;
-						}
-
-						LastIncludeLine = LineIndex - 1;
-						bAddSpace = true;
-						break;  // Unknown code
-					}
-
-					if (AddBeforeLine == -1 && LastIncludeLine != -1)
-					{
-						AddBeforeLine = LastIncludeLine + 1;
-					}
-
-					if (AddBeforeLine == -1)
-					{
-						AddBeforeLine = LastLine;
-					}
-
-					EditPoint.LineDown(AddBeforeLine - 1);
-					EditPoint.Insert($"{PasteString}\r\n");
-					if (bAddSpace)
-					{
-						EditPoint.Insert("\r\n");
-					}
-					Doc.Selection.MoveTo(AddBeforeLine, 0);
-					Doc.Selection.SelectLine();
-					EditPoint.TryToShow();
+					InsertIncludeIntoDocument((FileItem)(ActiveFilesListBox.SelectedItems[0]));
 				}
 				return;
 			}
@@ -519,7 +658,9 @@ namespace UnrealVS
 			if (E.Key == Key.Enter)
 			{
 				if (OpenSelectedFile())
-					HidePanel();
+				{
+					HidePanelAfterSelection();
+				}
 				return;
 			}
 
@@ -571,7 +712,9 @@ namespace UnrealVS
 		private void ListView_MouseDoubleClick(object Sender, MouseButtonEventArgs E)
 		{
 			if (OpenSelectedFile())
-				HidePanel();
+			{
+				HidePanelAfterSelection();
+			}
 		}
 
 		private bool OpenSelectedFile()
@@ -587,12 +730,23 @@ namespace UnrealVS
 				ToOpen.Add(Item.File);
 			}
 
+			ParseLineNumber(FilterEditBox.Text, out int LineNumber);
+
 			ThreadHelper.JoinableTaskFactory.Run(async delegate
 			{
 				await ThreadHelper.JoinableTaskFactory.SwitchToMainThreadAsync();
 				foreach (string F in ToOpen)
 				{
 					VsShellUtilities.OpenDocument(UnrealVSPackage.Instance, F);
+					if (LineNumber > 0)
+					{
+						Document ActiveDoc = UnrealVSPackage.Instance.DTE.ActiveDocument;
+						TextDocument TextDoc = ActiveDoc?.Object("TextDocument") as TextDocument;
+						if (TextDoc != null)
+						{
+							TextDoc.Selection.GotoLine(LineNumber, true);
+						}
+					}
 				}
 			});
 
@@ -605,15 +759,12 @@ namespace UnrealVS
 			{
 				return null;
 			}
-			string IncludePath = Item.File;
-			IncludePath = IncludePath.Replace("\\", "/");
-			string IncludePathLwr = IncludePath.ToLower();
-			int IncludeStartIndex = Math.Max(Math.Max(IncludePathLwr.LastIndexOf("/public/"), IncludePathLwr.LastIndexOf("/private/")), IncludePathLwr.LastIndexOf("/classes/"));
-			if (IncludeStartIndex != -1)
+
+			if (CopyIncludePath.BuildIncludePath(Item.File, out string IncludePath))
 			{
-				IncludePath = IncludePath.Substring(IncludePath.IndexOf('/', IncludeStartIndex + 1) + 1);
+				return IncludePath;
 			}
-			return IncludePath;
+			return null;
 		}
 
 		private bool CreateIncludePath()
@@ -710,12 +861,52 @@ namespace UnrealVS
 			public int Index;
 		}
 
+		// Parses an optional line number suffix from a filter string.
+		// Supports "name:123" and "name(123)" syntax.
+		// Returns the filter text with the line number portion removed, and sets LineNumber (-1 if absent).
+		private static string ParseLineNumber(string FilterText, out int LineNumber)
+		{
+			LineNumber = -1;
+			string Trimmed = FilterText.TrimEnd();
+
+			// "name(123)" pattern
+			if (Trimmed.EndsWith(")"))
+			{
+				int OpenParen = Trimmed.LastIndexOf('(');
+				if (OpenParen != -1)
+				{
+					string Inside = Trimmed.Substring(OpenParen + 1, Trimmed.Length - OpenParen - 2).Trim();
+					if (Inside.Length > 0 && int.TryParse(Inside, out int N) && N > 0)
+					{
+						LineNumber = N;
+						return Trimmed.Substring(0, OpenParen);
+					}
+				}
+			}
+
+			// "name:123" pattern
+			int ColonIndex = Trimmed.LastIndexOf(':');
+			if (ColonIndex != -1)
+			{
+				string After = Trimmed.Substring(ColonIndex + 1).Trim();
+				if (After.Length > 0 && int.TryParse(After, out int N) && N > 0)
+				{
+					LineNumber = N;
+					return Trimmed.Substring(0, ColonIndex);
+				}
+			}
+
+			return FilterText;
+		}
+
 		private void RefreshListView(ListBox ListView, FileItem[] Source)
 		{
 			List<FileItem> FileItems = (List<FileItem>)ListView.ItemsSource;
 			FileItems.Clear();
 
-			string FilterStr = FilterEditBox.Text;
+			string FilterStr = ParseLineNumber(FilterEditBox.Text, out int _LineNumber);
+			// Strip line-number delimiter characters so they don't interfere with name matching
+			FilterStr = FilterStr.Replace(':', ' ').Replace('(', ' ').Replace(')', ' ');
 			string[] Filter = FilterStr.Split(new char[] { ' ' }, StringSplitOptions.RemoveEmptyEntries);
 
 			HashSet<FileItem> Selections = new HashSet<FileItem>();
@@ -916,7 +1107,9 @@ namespace UnrealVS
 
 				Dictionary<string, FileItem> Lookup = new Dictionary<string, FileItem>();
 				foreach (FileItem Item in AllFileItems)
-					try { Lookup.Add(Item.File, Item); } catch (Exception) { }
+				{
+					Lookup[Item.File] = Item;
+				}
 				AllFileItemsLookup = Lookup;
 
 				LoadSolutionSettings();
@@ -936,14 +1129,20 @@ namespace UnrealVS
 			Timer.Start();
 		}
 
+		internal void HandleShow()
+		{
+			FilterEditBox.SelectAll();
+			FilterEditBox.Focus();
+			bClearFilter = true;
+		}
+
 		private void FileBrowserWindowControl_IsVisibleChanged(object Sender, DependencyPropertyChangedEventArgs E)
 		{
 			if (!(bool)E.NewValue)
 			{
 				return;
 			}
-			FilterEditBox.SelectAll();
-			bClearFilter = true;
+			HandleShow();
 			//FilesListTab.SelectedIndex = 0;
 
 			if (AllFileItems.Length == 0)
@@ -955,6 +1154,55 @@ namespace UnrealVS
 			AsyncFocusSelectedItem();
 		}
 
+		private void ColumnHeader_Click(object sender, RoutedEventArgs e)
+		{
+			var header = e.OriginalSource as GridViewColumnHeader;
+			if (header == null)
+				return;
+
+			string property = header.Tag as string;
+			if (property == null)
+				return;
+
+			var listBox = (ListBox)sender;
+
+			string originalContent = header.Column?.Header as string ?? property;
+			ListSortDirection direction = ListSortDirection.Ascending;
+
+			// for each list box, we save the column and the direction it was sorted.
+			// if a columns was sorted, then restore the previous sorted header.
+			if (ColumnSortState.TryGetValue(listBox, out var prev))
+			{
+				prev.Header.Content = prev.OriginalContent;
+				// new direction?
+				if (prev.Header == header && prev.Direction == ListSortDirection.Ascending)
+					direction = ListSortDirection.Descending;
+			}
+
+			header.Content = originalContent + (direction == ListSortDirection.Ascending ? " ▲" : " ▼");
+			ColumnSortState[listBox] = (header, originalContent, direction);
+
+			var view = CollectionViewSource.GetDefaultView(listBox.ItemsSource);
+			if (view != null)
+			{
+				view.SortDescriptions.Clear();
+				view.SortDescriptions.Add(new SortDescription(property, direction));
+			}
+
+			// Scroll the selected item into view
+			if (listBox.SelectedItem != null)
+			{
+				EventHandler scrollHandler = null;
+				scrollHandler = (s2, e2) =>
+				{
+					listBox.LayoutUpdated -= scrollHandler;
+					if (listBox.SelectedItem != null)
+						listBox.ScrollIntoView(listBox.SelectedItem);
+				};
+				listBox.LayoutUpdated += scrollHandler;
+			}
+		}
+
 		private void HidePanel()
 		{
 			ThreadHelper.ThrowIfNotOnUIThread();
@@ -962,6 +1210,14 @@ namespace UnrealVS
 			ToolWindowFrame.Hide();
 		}
 
+		private void HidePanelAfterSelection()
+		{
+			if (!UnrealVSPackage.Instance.OptionsPage.FileBrowserKeepOpenAfterSelection)
+			{
+				HidePanel();
+			}
+		}
+
 		private class FileBrowserSettings
 		{
 			public List<string> Bookmarks { get; set; } = new List<string>();
@@ -1098,14 +1354,15 @@ namespace UnrealVS
 
 		private class SolutionTraverser
 		{
-			private SortedDictionary<string, FileItem> SortedFileItems = new SortedDictionary<string, FileItem>();
-			private Dictionary<string, SortedDictionary<string, FileItem>> CollidingFileItems = new Dictionary<string, SortedDictionary<string, FileItem>>();
+			private List<FileItem> FileItems = new List<FileItem>();
 			private Dictionary<string, string> ModuleRoots = new Dictionary<string, string>();
 
 			private Stack<StackItem> ItemStack = new Stack<StackItem>();
 			private int ProjectIndex;
 			public int HandledItemCount { get; private set; }
 
+			private readonly string[] Extensions = new string[] { ".build.cs", ".target.cs", ".uproject", ".uplugin" };
+
 			public FileItem[] Update()
 			{
 				Stack<StackItem> TraversalStack = ItemStack;
@@ -1150,9 +1407,9 @@ namespace UnrealVS
 						ProjectItem ProjectItem = CurrentItem.Items.Item(CurrentItem.Index + 1);
 						++TraverseCounter;
 
-						if (ProjectItem.FileCount != 0)
+						for (short FileIndex = 0; FileIndex < ProjectItem.FileCount; ++FileIndex)
 						{
-							string FileName = ProjectItem.FileNames[1];
+							string FileName = ProjectItem.FileNames[(short)(FileIndex + 1)];
 							if (FileName != null)
 							{
 								string Name = ProjectItem.Name;
@@ -1160,41 +1417,23 @@ namespace UnrealVS
 								{
 									if (!FileName.EndsWith("\\")) // Skip folders
 									{
-										FileItem NewFileItem = new FileItem() { Name = Name, File = FileName, Project = CurrentItem.Project.Name };
-										FileItem ExistingFileItem;
-										if (SortedFileItems.TryGetValue(Name, out ExistingFileItem))
+										FileItem NewFileItem = new FileItem() { Name = Name, File = FileName, Module = CurrentItem.Project.Name };
+										FileItems.Add(NewFileItem);
+
+										foreach (string Extension in Extensions)
 										{
-											if (ExistingFileItem.File != FileName)
+											if (FileName.EndsWith(Extension, StringComparison.OrdinalIgnoreCase))
 											{
-												SortedDictionary<string, FileItem> ColList;
-												if (!CollidingFileItems.TryGetValue(Name, out ColList))
+												int LastBackslash = FileName.LastIndexOf('\\');
+												string ModuleDir = FileName.Substring(0, LastBackslash);
+												if (!ModuleRoots.ContainsKey(ModuleDir))
 												{
-													ColList = new SortedDictionary<string, FileItem>();
-													ColList.Add(ExistingFileItem.File, ExistingFileItem);
-													CollidingFileItems[Name] = ColList;
+													string ModuleName = FileName.Substring(LastBackslash + 1, FileName.Length - LastBackslash - Extension.Length - 1);
+													ModuleRoots.Add(ModuleDir, ModuleName);
 												}
-												if (!ColList.ContainsKey(NewFileItem.File))
-												{
-													ColList.Add(NewFileItem.File, NewFileItem);
-												}
-											}
-										}
-										else
-											SortedFileItems.Add(Name, NewFileItem);
-
-										if (FileName.EndsWith(".build.cs", StringComparison.OrdinalIgnoreCase))
-										{
-											int LastBackslash = FileName.LastIndexOf('\\');
-											string ModuleDir = FileName.Substring(0, LastBackslash);
-											string ModuleName = FileName.Substring(LastBackslash + 1, FileName.Length - LastBackslash - 10);
-											if (!ModuleRoots.ContainsKey(ModuleDir))
-											{
-												ModuleRoots.Add(ModuleDir, ModuleName);
+												break;
 											}
-											//else if (ModuleRoots[moduleDir] != moduleName)
-											//    Logging.WriteLine(file);
 										}
-
 										++HandledItemCount;
 									}
 								}
@@ -1221,80 +1460,33 @@ namespace UnrealVS
 					++ProjectIndex;
 				}
 
-				// We need to rename all collisions
-				foreach (KeyValuePair<string, SortedDictionary<string, FileItem>> CollisionKv in CollidingFileItems)
+				// Assign module to all file items
+				foreach (FileItem Item in FileItems)
 				{
-					SortedFileItems.Remove(CollisionKv.Key);
-					Dictionary<string, int> UsedParents = new Dictionary<string, int>();
-
-					foreach (KeyValuePair<string, FileItem> ItemKv in CollisionKv.Value)
+					var Module = GetModuleForFile(Item.File);
+					if (Module != null)
 					{
-						FileItem Item = ItemKv.Value;
-						string Parent = Item.Project;
-						string FilePath = Item.File;
-						int LastBackslash = FilePath.LastIndexOf('\\');
-						while (LastBackslash != -1)
-						{
-							FilePath = FilePath.Substring(0, LastBackslash);
-							string ModuleName;
-							if (ModuleRoots.TryGetValue(FilePath, out ModuleName))
-							{
-								Parent = ModuleName;
-								break;
-							}
-							LastBackslash = FilePath.LastIndexOf('\\');
-						}
+						Item.Module = Module;
+					}
+				}
 
-						string SortName = Item.Name + $" ({Parent})";
+				return FileItems.ToArray();
+			}
 
-						int Counter = 0;
-						if (UsedParents.TryGetValue(Parent, out Counter))
-						{
-							if (Counter == 0)
-							{
-								FileItem AlreadyAddedItem = SortedFileItems[SortName];
-								SortedFileItems.Remove(SortName);
-								AlreadyAddedItem.Name = Item.Name + $" ({Parent} 0)";
-								SortedFileItems.Add(AlreadyAddedItem.Name, AlreadyAddedItem);
-							}
-							++Counter;
-							UsedParents[Parent] = Counter;
-							Item.Name = Item.Name + $" ({Parent} {Counter})";
-							SortedFileItems.Add(Item.Name, Item);
-						}
-						else
-						{
-							UsedParents.Add(Parent, 0);
-							Item.Name = SortName;
-							SortedFileItems.Add(Item.Name, Item);
-						}
+			private string GetModuleForFile(string FilePath)
+			{
+				int LastBackslash = FilePath.LastIndexOf('\\');
+				while (LastBackslash > 0)
+				{
+					string dir = FilePath.Substring(0, LastBackslash);
+					if (ModuleRoots.TryGetValue(dir, out string ModuleName))
+					{
+						return ModuleName;
 					}
+					LastBackslash = dir.LastIndexOf('\\');
 				}
-
-				FileItem[] Result = SortedFileItems.Values.ToArray();
-
-				// Sort so header is before source file.. since most of the time you are opening header files.
-				/*
-                for (int i=0, e=result.Length; i!=e; ++i)
-                {
-                    var name = result[i].Name;
-					if (name.EndsWith(".h") && i > 0)
-                    {
-						var prev = result[i - 1];
-						var prevName = prev.Name;
-                        if (name.Length + 2 == prevName.Length)
-                        {
-                            if (name.AsSpan(0, name.Length - 1).SequenceEqual(prevName.AsSpan(0, prevName.Length - 3))) // swap h and cpp
-                            {
-                                result[i - 1] = result[i];
-                                result[i] = prev;
-							}
-                        }
-                    }
-                }
-                */
-				return Result;
+				return null;
 			}
 		}
 	}
-}
+}
\ No newline at end of file
diff --git a/Engine/Source/Programs/UnrealVS/UnrealVS.Shared/UnrealVSOptions.cs b/Engine/Source/Programs/UnrealVS/UnrealVS.Shared/UnrealVSOptions.cs
--- a/Engine/Source/Programs/UnrealVS/UnrealVS.Shared/UnrealVSOptions.cs
+++ b/Engine/Source/Programs/UnrealVS/UnrealVS.Shared/UnrealVSOptions.cs
@@ -62,6 +62,11 @@ namespace UnrealVS
 		[Description("Will include name of folder that UE5.sln is located in to Visual Studio window title.")]
 		public bool IncludeFolderInUE5SolutionName { get; set; } = true;
 
+		[Category("File Browser")]
+		[DisplayName("Keep Open After Selection")]
+		[Description("When enabled, the File Browser window stays open after opening a file instead of closing automatically.")]
+		public bool FileBrowserKeepOpenAfterSelection { get; set; } = true;
+
 		protected override void OnApply(PageApplyEventArgs e)
 		{
 			base.OnApply(e);
