[UnrealVS] Add Cycle Related File command. It finds and opens the next file with the same name (extension can be different). It starts looking from the module folder, if it exists.

#Synced-CL 52420744

diff --git a/Engine/Extras/UnrealVS/VS2022/AtomFeed.xml b/Engine/Extras/UnrealVS/VS2022/AtomFeed.xml
--- a/Engine/Extras/UnrealVS/VS2022/AtomFeed.xml
+++ b/Engine/Extras/UnrealVS/VS2022/AtomFeed.xml
@@ -2,13 +2,13 @@
 <feed xmlns="http://www.w3.org/2005/Atom">
   <title type="text">UnrealVS Gallery</title>
   <id>uuid:45b4571a-2981-49f1-9f70-04cfa96db5c8;id=1</id>
-  <updated>2026-04-07T00:00:00Z</updated>
+  <updated>2026-04-08T00:00:00Z</updated>
   <entry>
     <id>ddbf523f-7eb6-4887-bd51-85a714ff87eb</id>
     <title type="text">UnrealVS</title>
     <summary type="text">Unreal Engine extension for Visual Studio</summary>
     <published>2022-10-21T10:01:57-07:00</published>
-    <updated>2026-04-07T00:00:00+00:00</updated>
+    <updated>2026-04-08T00:00:00+00:00</updated>
     <author>
       <name>Epic Games, Inc.</name>
     </author>
@@ -19,7 +19,7 @@
     <content type="application/octet-stream" src="UnrealVS.vsix" />
     <Vsix xmlns:xsd="http://www.w3.org/2001/XMLSchema" xmlns:xsi="http://www.w3.org/2001/XMLSchema-instance" xmlns="http://schemas.microsoft.com/developer/vsx-syndication-schema/2010">
       <Id>ddbf523f-7eb6-4887-bd51-85a714ff87eb</Id>
-      <Version>1.85.0.0</Version>
+      <Version>1.85.1.0</Version>
       <Rating xsi:nil="true" />
       <RatingCount xsi:nil="true" />
       <DownloadCount xsi:nil="true" />
diff --git a/Engine/Source/Programs/UnrealVS/ReleaseNotes.txt b/Engine/Source/Programs/UnrealVS/ReleaseNotes.txt
--- a/Engine/Source/Programs/UnrealVS/ReleaseNotes.txt
+++ b/Engine/Source/Programs/UnrealVS/ReleaseNotes.txt
@@ -1,7 +1,10 @@
+1.86.0.0
+* Added Cycle Related File command. This command will cycle through files related to the current file. For example, if you are in a .cpp file, it will open the corresponding .h file, and vice versa.
+
 1.85.0.0
-* Add option to not close the File Browser window after a file is select.
-* Add the module and file path to File Browser.
-* Add s context menu to File Browser.
+* Added option to not close the File Browser window after a file is select.
+* Added the module and file path to File Browser.
+* Added a context menu to File Browser.
 
 1.84.9.6
 * Usability improvements for startup project selector
diff --git a/Engine/Source/Programs/UnrealVS/UnrealVS.2022/Properties/AssemblyInfo.cs b/Engine/Source/Programs/UnrealVS/UnrealVS.2022/Properties/AssemblyInfo.cs
--- a/Engine/Source/Programs/UnrealVS/UnrealVS.2022/Properties/AssemblyInfo.cs
+++ b/Engine/Source/Programs/UnrealVS/UnrealVS.2022/Properties/AssemblyInfo.cs
@@ -31,6 +31,6 @@ using System.Runtime.InteropServices;
 // You can specify all the values or you can default the Revision and Build Numbers 
 // by using the '*' as shown below:
 
-[assembly: AssemblyVersion("1.85.0.0")]
-[assembly: AssemblyFileVersion("1.85.0.0")]
+[assembly: AssemblyVersion("1.86.0.0")]
+[assembly: AssemblyFileVersion("1.86.0.0")]
 
diff --git a/Engine/Source/Programs/UnrealVS/UnrealVS.2022/UnrealVS.2022.csproj b/Engine/Source/Programs/UnrealVS/UnrealVS.2022/UnrealVS.2022.csproj
--- a/Engine/Source/Programs/UnrealVS/UnrealVS.2022/UnrealVS.2022.csproj
+++ b/Engine/Source/Programs/UnrealVS/UnrealVS.2022/UnrealVS.2022.csproj
@@ -122,6 +122,8 @@
     </Resource>
     <Resource Include="Resources\InsertInclude.png">
     </Resource>
+    <Resource Include="Resources\CycleRelatedFiles.png">
+    </Resource>
     <Resource Include="Resources\UnrealIcon.png">
     </Resource>
     <None Include="source.extension.vsixmanifest" />
diff --git a/Engine/Source/Programs/UnrealVS/UnrealVS.Shared/CompileSingleFile.cs b/Engine/Source/Programs/UnrealVS/UnrealVS.Shared/CompileSingleFile.cs
--- a/Engine/Source/Programs/UnrealVS/UnrealVS.Shared/CompileSingleFile.cs
+++ b/Engine/Source/Programs/UnrealVS/UnrealVS.Shared/CompileSingleFile.cs
@@ -113,7 +113,18 @@ namespace UnrealVS
 			}
 		}
 
-		static string FindModuleForFile(string FileName, string RootDirectory)
+		static public string FindModuleForFile(string FileName, string RootDirectory)
+		{
+			string ModuleBuildFilePath = FindBuildFilePathForFile(FileName, RootDirectory);
+			if (ModuleBuildFilePath != null)
+			{
+				// *.Build.cs -> *.Build -> *
+				return Path.GetFileNameWithoutExtension(Path.GetFileNameWithoutExtension(ModuleBuildFilePath));
+			}
+			return null;
+		}
+
+		static public string FindBuildFilePathForFile(string FileName, string RootDirectory)
 		{
 			for (string Dir = Path.GetDirectoryName(FileName);
 				Dir != null && !String.Equals(Dir, RootDirectory, StringComparison.OrdinalIgnoreCase);
@@ -122,8 +133,7 @@ namespace UnrealVS
 				string BuildFile = Directory.EnumerateFiles(Dir, "*.Build.cs", SearchOption.TopDirectoryOnly).FirstOrDefault();
 				if (BuildFile != null)
 				{
-					// *.Build.cs -> *.Build -> *
-					return Path.GetFileNameWithoutExtension(Path.GetFileNameWithoutExtension(BuildFile));
+					return BuildFile;
 				}
 			}
 			return null;
diff --git a/Engine/Source/Programs/UnrealVS/UnrealVS.Shared/CycleRelatedFiles.cs b/Engine/Source/Programs/UnrealVS/UnrealVS.Shared/CycleRelatedFiles.cs
new file mode 100644
--- /dev/null
+++ b/Engine/Source/Programs/UnrealVS/UnrealVS.Shared/CycleRelatedFiles.cs
@@ -0,0 +1,88 @@
+// Copyright Epic Games, Inc. All Rights Reserved.
+
+using EnvDTE;
+using Microsoft.VisualStudio.Shell;
+using Microsoft.VisualStudio.Shell.Interop;
+using System;
+using System.ComponentModel.Design;
+using System.IO;
+
+namespace UnrealVS
+{
+	/// Cycles through files that share the same base name as the active document
+	/// (e.g. MyClass.h → MyClass.cpp → MyClass.inl → …).
+	/// The search root is the module directory found by walking up from the current
+	/// file and locating the folder that contains a .build.cs file.
+	class CycleRelatedFiles
+	{
+		const int CycleRelatedFilesButtonID = 0x1339;
+
+		public CycleRelatedFiles()
+		{
+			CommandID CommandID = new CommandID(GuidList.UnrealVSCmdSet, CycleRelatedFilesButtonID);
+			MenuCommand Command = new MenuCommand(new EventHandler(CycleRelatedFilesHandler), CommandID);
+			UnrealVSPackage.Instance.MenuCommandService.AddCommand(Command);
+		}
+
+		[global::System.Diagnostics.CodeAnalysis.SuppressMessageAttribute("", "VSTHRD010")]
+		private void CycleRelatedFilesHandler(object Sender, EventArgs Args)
+		{
+			ThreadHelper.ThrowIfNotOnUIThread();
+
+			Document ActiveDoc = UnrealVSPackage.Instance.DTE.ActiveDocument;
+			if (ActiveDoc == null)
+			{
+				Logging.WriteLine("CycleRelatedFiles: ActiveDocument not found");
+				return;
+			}
+
+			string CurrentFilePath = ActiveDoc.FullName;
+			string BaseFileName = Path.GetFileNameWithoutExtension(CurrentFilePath);
+			string CurrentDir = Path.GetDirectoryName(CurrentFilePath);
+
+			// Find the module root: walk up until we find a directory containing a .build.cs file.
+			string SolutionFileName = UnrealVSPackage.Instance.DTE.Solution.FileName;
+			if (SolutionFileName == null)
+			{
+				Logging.WriteLine("CycleRelatedFiles: Solution not found");
+				return;
+			}
+			string ModuleRoot = CompileSingleFile.FindBuildFilePathForFile(CurrentFilePath, Path.GetDirectoryName(SolutionFileName));
+			if (string.IsNullOrEmpty(ModuleRoot))
+			{
+				ModuleRoot = CurrentDir;
+			}
+			else
+			{
+				ModuleRoot = Path.GetDirectoryName(ModuleRoot);
+			}
+
+			// Collect all files with the same base name anywhere under the module root.
+			string[] RelatedFiles;
+			try
+			{
+				RelatedFiles = Directory.GetFiles(ModuleRoot, BaseFileName + ".*", SearchOption.AllDirectories);
+			}
+			catch
+			{
+				Logging.WriteLine("CycleRelatedFiles: Inaccessible directory");
+				return;
+			}
+			if (RelatedFiles.Length <= 1)
+			{
+				Logging.WriteLine("CycleRelatedFiles: No other files");
+				return;
+			}
+
+			Array.Sort(RelatedFiles, StringComparer.OrdinalIgnoreCase);
+
+			int CurrentIndex = Array.FindIndex(RelatedFiles,
+				F => string.Equals(F, CurrentFilePath, StringComparison.OrdinalIgnoreCase));
+
+			// If not found, then CurrentIndex will be -1 , and NextIndex will wrap to 0, which is what we want.
+			int NextIndex = (CurrentIndex + 1) % RelatedFiles.Length;
+
+			VsShellUtilities.OpenDocument(UnrealVSPackage.Instance, RelatedFiles[NextIndex]);
+		}
+	}
+}
diff --git a/Engine/Source/Programs/UnrealVS/UnrealVS.Shared/UnrealVSPackage.cs b/Engine/Source/Programs/UnrealVS/UnrealVS.Shared/UnrealVSPackage.cs
--- a/Engine/Source/Programs/UnrealVS/UnrealVS.Shared/UnrealVSPackage.cs
+++ b/Engine/Source/Programs/UnrealVS/UnrealVS.Shared/UnrealVSPackage.cs
@@ -292,6 +292,9 @@ namespace UnrealVS
 			// Create 'FileBrowser' instance
 			FileBrowser = new FileBrowser();
 
+			// Create 'CycleRelatedFiles' instance
+			CycleRelatedFiles = new CycleRelatedFiles();
+
 			// Create 'UbaVisualizer' instance
 			UbaVisualizer = new UbaVisualizer();
 
@@ -1113,6 +1116,9 @@ namespace UnrealVS
 		/// FileBrowser feature
 		private FileBrowser FileBrowser;
 
+		/// CycleRelatedFiles feature
+		private CycleRelatedFiles CycleRelatedFiles;
+
 		/// UbaVisualizer feature
 		private UbaVisualizer UbaVisualizer;
 
